Copy up Dev-Enterprise @ CL11133089

diff --git a/Engine/Plugins/Enterprise/DatasmithGLTFImporter/Source/DatasmithGLTFTranslator/Private/DatasmithGLTFTranslator.cpp b/Engine/Plugins/Enterprise/DatasmithGLTFImporter/Source/DatasmithGLTFTranslator/Private/DatasmithGLTFTranslator.cpp
--- a/Engine/Plugins/Enterprise/DatasmithGLTFImporter/Source/DatasmithGLTFTranslator/Private/DatasmithGLTFTranslator.cpp
+++ b/Engine/Plugins/Enterprise/DatasmithGLTFImporter/Source/DatasmithGLTFTranslator/Private/DatasmithGLTFTranslator.cpp
@@ -67,7 +67,7 @@ bool FDatasmithGLTFTranslator::LoadScene(TSharedRef<IDatasmithScene> OutScene)
 {
 	OutScene->SetHost(TEXT("GLTFTranslator"));
 
-    Importer = MakeShared<FDatasmithGLTFImporter>(OutScene, ImportOptions.Get());
+    Importer = MakeShared<FDatasmithGLTFImporter>(OutScene, GetOrCreateGLTFImportOptions().Get());
 
 	const FString& FilePath = GetSource().GetSourceFile();
 	if(!Importer->OpenFile(FilePath))
@@ -167,12 +167,7 @@ bool FDatasmithGLTFTranslator::LoadLevelSequence(const TSharedRef<IDatasmithLeve
 
 void FDatasmithGLTFTranslator::GetSceneImportOptions(TArray<TStrongObjectPtr<UObject>>& Options)
 {
-	if (!ImportOptions.IsValid())
-	{
-		ImportOptions = Datasmith::MakeOptions<UDatasmithGLTFImportOptions>();
-	}
-
-	Options.Add(ImportOptions);
+	Options.Add(GetOrCreateGLTFImportOptions());
 }
 
 void FDatasmithGLTFTranslator::SetSceneImportOptions(TArray<TStrongObjectPtr<UObject>>& Options)
@@ -188,6 +183,15 @@ void FDatasmithGLTFTranslator::SetSceneImportOptions(TArray<TStrongObjectPtr<UOb
 
 	if (Importer.IsValid())
 	{
-		Importer->SetImportOptions(ImportOptions.Get());
+		Importer->SetImportOptions(GetOrCreateGLTFImportOptions().Get());
+	}
+}
+
+TStrongObjectPtr<UDatasmithGLTFImportOptions>& FDatasmithGLTFTranslator::GetOrCreateGLTFImportOptions()
+{
+	if (!ImportOptions.IsValid())
+	{
+		ImportOptions = Datasmith::MakeOptions<UDatasmithGLTFImportOptions>();
 	}
+	return ImportOptions;
 }
diff --git a/Engine/Plugins/Enterprise/DatasmithGLTFImporter/Source/DatasmithGLTFTranslator/Private/DatasmithGLTFTranslator.h b/Engine/Plugins/Enterprise/DatasmithGLTFImporter/Source/DatasmithGLTFTranslator/Private/DatasmithGLTFTranslator.h
--- a/Engine/Plugins/Enterprise/DatasmithGLTFImporter/Source/DatasmithGLTFTranslator/Private/DatasmithGLTFTranslator.h
+++ b/Engine/Plugins/Enterprise/DatasmithGLTFImporter/Source/DatasmithGLTFTranslator/Private/DatasmithGLTFTranslator.h
@@ -24,6 +24,9 @@ public:
 	virtual void SetSceneImportOptions(TArray<TStrongObjectPtr<UObject>>& Options) override;
 	//~ End IDatasmithTranslator interface
 
+private:
+	TStrongObjectPtr<UDatasmithGLTFImportOptions>& GetOrCreateGLTFImportOptions();
+
 private:
     TStrongObjectPtr<UDatasmithGLTFImportOptions> ImportOptions;
     TSharedPtr<class FDatasmithGLTFImporter> Importer;
diff --git a/Engine/Plugins/Enterprise/DatasmithGLTFImporter/Source/DatasmithGLTFTranslator/Public/DatasmithGLTFTranslatorModule.h b/Engine/Plugins/Enterprise/DatasmithGLTFImporter/Source/DatasmithGLTFTranslator/Public/DatasmithGLTFTranslatorModule.h
--- a/Engine/Plugins/Enterprise/DatasmithGLTFImporter/Source/DatasmithGLTFTranslator/Public/DatasmithGLTFTranslatorModule.h
+++ b/Engine/Plugins/Enterprise/DatasmithGLTFImporter/Source/DatasmithGLTFTranslator/Public/DatasmithGLTFTranslatorModule.h
@@ -5,7 +5,7 @@
 #include "CoreMinimal.h"
 #include "Modules/ModuleManager.h"
 
-class IDatasmithGLTFTranslatorModule : public IModuleInterface
+class DATASMITHGLTFTRANSLATOR_API IDatasmithGLTFTranslatorModule : public IModuleInterface
 {
 public:
 	static const TCHAR* ModuleName;
diff --git a/Engine/Plugins/Enterprise/DatasmithIFCImporter/Source/DatasmithIFCTranslator/DatasmithIFCTranslator.Build.cs b/Engine/Plugins/Enterprise/DatasmithIFCImporter/Source/DatasmithIFCTranslator/DatasmithIFCTranslator.Build.cs
--- a/Engine/Plugins/Enterprise/DatasmithIFCImporter/Source/DatasmithIFCTranslator/DatasmithIFCTranslator.Build.cs
+++ b/Engine/Plugins/Enterprise/DatasmithIFCImporter/Source/DatasmithIFCTranslator/DatasmithIFCTranslator.Build.cs
@@ -6,27 +6,27 @@ using System.IO;
 
 namespace UnrealBuildTool.Rules
 {
-    public class DatasmithIFCTranslator : ModuleRules
-    {
-        public DatasmithIFCTranslator(ReadOnlyTargetRules Target) : base(Target)
-        {
-            PrivateDependencyModuleNames.AddRange(
-                new string[]
-                {
-                    "Analytics",
-                    "CoreUObject",
-                    "DatasmithCore",
+	public class DatasmithIFCTranslator : ModuleRules
+	{
+		public DatasmithIFCTranslator(ReadOnlyTargetRules Target) : base(Target)
+		{
+			PrivateDependencyModuleNames.AddRange(
+				new string[]
+				{
+					"Analytics",
+					"CoreUObject",
+					"DatasmithCore",
 					"Engine",
 					"Json",
-                    "MeshDescription",
+					"MeshDescription",
 					"Projects",
 					"RawMesh",
-                    "Slate",
-                    "SlateCore",
+					"Slate",
+					"SlateCore",
 					"StaticMeshDescription",
-                    "XmlParser",
-                }
-            );
+					"XmlParser",
+				}
+			);
 
 			if (Target.Type == TargetType.Editor)
 			{
@@ -45,19 +45,26 @@ namespace UnrealBuildTool.Rules
 			}
 
 			PublicDependencyModuleNames.AddRange(
-                new string[]
-                {
-                    "Core",
-                    "DatasmithContent",
-                    "DatasmithTranslator"
-                }
-            );
+				new string[]
+				{
+					"Core",
+					"DatasmithContent",
+					"DatasmithTranslator"
+				}
+			);
 
 			if (Target.Platform == UnrealTargetPlatform.Win64)
 			{
 				RuntimeDependencies.Add("$(EngineDir)/Plugins/Enterprise/DatasmithIFCImporter/Binaries/Win64/ifcengine.dll");
 				PublicDelayLoadDLLs.Add("ifcengine.dll");
 			}
+
+			RuntimeDependencies.Add(@"$(EngineDir)\Plugins\Enterprise\DatasmithIFCImporter\Resources\IFC\IFC2X3_TC1.exp");
+			RuntimeDependencies.Add(@"$(EngineDir)\Plugins\Enterprise\DatasmithIFCImporter\Resources\IFC\IFC2X3-Settings.xml");
+			RuntimeDependencies.Add(@"$(EngineDir)\Plugins\Enterprise\DatasmithIFCImporter\Resources\IFC\IFC4.exp");
+			RuntimeDependencies.Add(@"$(EngineDir)\Plugins\Enterprise\DatasmithIFCImporter\Resources\IFC\IFC4_ADD2.exp");
+			RuntimeDependencies.Add(@"$(EngineDir)\Plugins\Enterprise\DatasmithIFCImporter\Resources\IFC\IFC4-Settings.xml");
+			RuntimeDependencies.Add(@"$(EngineDir)\Plugins\Enterprise\DatasmithIFCImporter\Resources\IFC\IFC4x1_FINAL.exp");
 		}
 	}
 }
diff --git a/Engine/Plugins/Enterprise/DatasmithIFCImporter/Source/DatasmithIFCTranslator/Public/DatasmithIFCTranslatorModule.h b/Engine/Plugins/Enterprise/DatasmithIFCImporter/Source/DatasmithIFCTranslator/Public/DatasmithIFCTranslatorModule.h
--- a/Engine/Plugins/Enterprise/DatasmithIFCImporter/Source/DatasmithIFCTranslator/Public/DatasmithIFCTranslatorModule.h
+++ b/Engine/Plugins/Enterprise/DatasmithIFCImporter/Source/DatasmithIFCTranslator/Public/DatasmithIFCTranslatorModule.h
@@ -7,7 +7,7 @@
 
 DECLARE_LOG_CATEGORY_EXTERN(LogIFCTranslator, Log, All);
 
-class IDatasmithIFCTranslatorModule : public IModuleInterface
+class DATASMITHIFCTRANSLATOR_API IDatasmithIFCTranslatorModule : public IModuleInterface
 {
 public:
 	static const TCHAR* ModuleName;
diff --git a/Engine/Plugins/Enterprise/LidarPointCloud/Config/BaseLidarPointCloud.ini b/Engine/Plugins/Enterprise/LidarPointCloud/Config/BaseLidarPointCloud.ini
new file mode 100644
--- /dev/null
+++ b/Engine/Plugins/Enterprise/LidarPointCloud/Config/BaseLidarPointCloud.ini
@@ -0,0 +1,28 @@
+[CoreRedirects]
+; Blueprints
++ClassRedirects=(OldName="/Script/PointCloudRuntime.PointCloud",NewName="LidarPointCloud")
++ClassRedirects=(OldName="/Script/PointCloudRuntime.PointCloudPluginComponent",NewName="LidarPointCloudComponent")
++ClassRedirects=(OldName="/Script/PointCloudRuntime.PointCloudActor",NewName="LidarPointCloudActor")
++ClassRedirects=(OldName="/Script/PointCloudRuntime.PointCloudBlueprintLibrary",NewName="LidarPointCloudBlueprintLibrary")
++StructRedirects=(OldName="/Script/PointCloudRuntime.PointCloudPoint",NewName="LidarPointCloudPoint")
++StructRedirects=(OldName="/Script/PointCloudRuntime.PointCloudTraceHit",NewName="LidarPointCloudTraceHit")
++StructRedirects=(OldName="/Script/PointCloudRuntime.PointCloudImportSettings_ASCII_Columns",NewName="LidarPointCloudImportSettings_ASCII_Columns")
++EnumRedirects=(OldName="/Script/PointCloudRuntime.EPointCloudColorationMode",NewName="ELidarPointCloudColorationMode")
++EnumRedirects=(OldName="/Script/PointCloudRuntime.EPointCloudSpriteShape",NewName="ELidarPointCloudSpriteShape")
++EnumRedirects=(OldName="/Script/PointCloudRuntime.EPointCloudDuplicateHandling",NewName="ELidarPointCloudDuplicateHandling")
++EnumRedirects=(OldName="/Script/PointCloudRuntime.EPointCloudAsyncMode",NewName="ELidarPointCloudAsyncMode")
+
+; Materials
++PackageRedirects=(OldName="/PointCloudPlugin/Materials/M_PointCloud",NewName="/LidarPointCloud/Materials/M_LidarPointCloud")
++PackageRedirects=(OldName="/PointCloudPlugin/Materials/M_PointCloud.M_PointCloud",NewName="/LidarPointCloud/Materials/M_LidarPointCloud.M_LidarPointCloud")
++PackageRedirects=(OldName="/PointCloudPlugin/Materials/M_PointCloud_Masked",NewName="/LidarPointCloud/Materials/M_LidarPointCloud_Masked")
++PackageRedirects=(OldName="/PointCloudPlugin/Materials/M_PointCloud_Masked.M_PointCloud_Masked",NewName="/LidarPointCloud/Materials/M_LidarPointCloud_Masked.M_LidarPointCloud_Masked")
++PackageRedirects=(OldName="/PointCloudPlugin/Materials/M_PP_EDL_MainPass",NewName="/LidarPointCloud/Materials/M_PP_EDL_MainPass")
++PackageRedirects=(OldName="/PointCloudPlugin/Materials/M_PP_EDL_MainPass.M_PP_EDL_MainPass",NewName="/LidarPointCloud/Materials/M_PP_EDL_MainPass.M_PP_EDL_MainPass")
++PackageRedirects=(OldName="/PointCloudPlugin/Materials/MF_EDF_SampleTexel_MainPass",NewName="/LidarPointCloud/Materials/MF_EDF_SampleTexel_MainPass")
++PackageRedirects=(OldName="/PointCloudPlugin/Materials/MF_EDF_SampleTexel_MainPass.MF_EDF_SampleTexel_MainPass",NewName="/LidarPointCloud/Materials/MF_EDF_SampleTexel_MainPass.MF_EDF_SampleTexel_MainPass")
++PackageRedirects=(OldName="/PointCloudPlugin/Materials/MF_EDF_SampleTexel_CustomPass",NewName="/LidarPointCloud/Materials/MF_EDF_SampleTexel_CustomPass")
++PackageRedirects=(OldName="/PointCloudPlugin/Materials/MF_EDF_SampleTexel_CustomPass.MF_EDF_SampleTexel_CustomPass",NewName="/LidarPointCloud/Materials/MF_EDF_SampleTexel_CustomPass.MF_EDF_SampleTexel_CustomPass")
+
+; Package
++PackageRedirects=(MatchSubstring=true,OldName="/PointCloudRuntime",NewName="/LidarPointCloudRuntime")
\ No newline at end of file
diff --git a/Engine/Plugins/Enterprise/LidarPointCloud/Shaders/Private/LidarPointCloudVertexFactory.ush b/Engine/Plugins/Enterprise/LidarPointCloud/Shaders/Private/LidarPointCloudVertexFactory.ush
new file mode 100644
--- /dev/null
+++ b/Engine/Plugins/Enterprise/LidarPointCloud/Shaders/Private/LidarPointCloudVertexFactory.ush
@@ -0,0 +1,390 @@
+// Copyright Epic Games, Inc. All Rights Reserved.
+
+#include "/Engine/Private/Common.ush"
+#include "/Engine/Private/VertexFactoryCommon.ush"
+
+/*=============================================================================
+	LidarPointCloudVertexFactory.usf: LiDAR Point Cloud vertex factory shader code.
+=============================================================================*/
+
+// First half is normal, second half has inverted V
+static float2 PrecomputedUV_Quad[8] = { float2(-0.5, -0.5), float2(0.5, -0.5), float2(0.5, 0.5), float2(-0.5, 0.5), float2(0, 1), float2(1, 1), float2(1, 0), float2(0, 0) };
+
+static float3 LODColors[] = {
+	float3(1, 1, 1), float3(1, 0, 0), float3(0, 1, 0), float3(0, 0, 1), float3(1, 1, 0), float3(1, 0, 1), float3(0, 1, 1),
+	float3(0.66, 0.66, 0.66), float3(0.66, 0, 0), float3(0, 0.66, 0), float3(0, 0, 0.66), float3(0.66, 0.66, 0), float3(0.66, 0, 0.66), float3(0, 0.66, 0.66),
+	float3(0.33, 0.33, 0.33), float3(0.33, 0, 0), float3(0, 0.33, 0), float3(0, 0, 0.33), float3(0.33, 0.33, 0), float3(0.33, 0, 0.33), float3(0, 0.33, 0.33)
+};
+
+// Used for instance handling
+uint IndexDivisor;
+
+// Used for Sprite Size
+float VDMultiplier;
+uint SizeOffset;
+float RootCellSize;
+float SpriteSizeMultiplier;
+
+// Needed for WPO calculations
+float3 ViewRightVector;
+float3 ViewUpVector;
+
+// Used for coloration override
+float3 BoundsSize;
+float3 BoundsOffset;
+half3 ElevationColorBottom;
+half3 ElevationColorTop;
+
+// Color Adjustment
+half4 Offset;
+half4 Contrast;
+half4 Saturation;
+half4 Gamma;
+half3 Tint;
+float IntensityInfluence;
+
+// Flags
+uint bUseCircle;
+uint bUseLODColoration;
+uint bUseColorOverride;
+uint bUseElevationColor;
+uint bUseClassification;
+
+float4 ClassificationColors[32];
+
+StructuredBuffer<float> DataBuffer;
+
+#if USE_INSTANCING
+	// Instanced stereo sets the eye index explicitly for instanced geometry
+	#if INSTANCED_STEREO
+		uint InstancedEyeIndex;
+	#endif
+#endif
+
+struct FVertexFactoryInput
+{
+    float4  Position    : ATTRIBUTE0;
+	uint    VertexId    : SV_VertexID;
+};
+
+struct FPositionOnlyVertexFactoryInput
+{
+	float4  Position    : ATTRIBUTE0;
+	uint    VertexId    : SV_VertexID;
+};
+
+struct FPositionAndNormalOnlyVertexFactoryInput
+{
+	float4  Position : ATTRIBUTE0;
+	float4	Normal : ATTRIBUTE2;
+	uint    VertexId : SV_VertexID;
+};
+
+struct FVertexFactoryInterpolantsVSToPS
+{
+	TANGENTTOWORLD_INTERPOLATOR_BLOCK
+
+	half4	Color : COLOR0;
+	float2	TexCoords : TEXCOORD0;
+
+#if INSTANCED_STEREO
+	nointerpolation uint EyeIndex : PACKED_EYE_INDEX;
+#endif
+};
+
+struct FVertexFactoryIntermediates
+{
+    half3x3 TangentToLocal;
+    half3x3 TangentToWorld;
+    half TangentToWorldSign;
+    half4 Color;
+	float3 Position;
+	half2 Scale;
+};
+
+struct FPointCloudInstanceData
+{
+	float3 Position;
+	uint Color;
+	half VirtualDepth;
+	half2 Scale;
+};
+
+FPointCloudInstanceData GetInstanceData(uint VertexId)
+{
+	FPointCloudInstanceData InstanceData = (FPointCloudInstanceData)0;
+
+	uint PointId = VertexId / IndexDivisor;
+	uint Idx = PointId * 4;
+
+	InstanceData.Position = float3(DataBuffer[Idx], DataBuffer[Idx + 1], DataBuffer[Idx + 2]);
+	InstanceData.Color = asuint(DataBuffer[Idx + 3]);
+	InstanceData.VirtualDepth = half((0x000000FF & (asuint(DataBuffer[SizeOffset + PointId / 4]) >> (PointId % 4 * 8)))) * VDMultiplier;
+	InstanceData.Scale = RootCellSize * PrecomputedUV_Quad[VertexId % 4] / pow(2, InstanceData.VirtualDepth);
+
+	return InstanceData;
+}
+
+half3x3 CalcTangentToWorldNoScale(in half3x3 TangentToLocal)
+{
+	half3x3 LocalToWorld = GetLocalToWorld3x3();
+	half3 InvScale = Primitive.InvNonUniformScaleAndDeterminantSign.xyz;
+	LocalToWorld[0] *= InvScale.x;
+	LocalToWorld[1] *= InvScale.y;
+	LocalToWorld[2] *= InvScale.z;
+	return mul(TangentToLocal, LocalToWorld);
+}
+
+float3 ProcessPosition(float3 Position, half2 Scale) { return Position + (ViewRightVector * Scale.x + ViewUpVector * Scale.y) * SpriteSizeMultiplier; }
+
+FVertexFactoryIntermediates GetVertexFactoryIntermediates(FVertexFactoryInput Input)
+{
+	FVertexFactoryIntermediates Intermediates = (FVertexFactoryIntermediates)0;
+	FPointCloudInstanceData InstanceData = GetInstanceData(Input.VertexId);
+	
+	Intermediates.Position = InstanceData.Position;
+	Intermediates.Scale = InstanceData.Scale;
+
+	Intermediates.TangentToLocal = half3x3(1, 0, 0, 0, 1, 0, 0, 0, 1);
+	Intermediates.TangentToWorldSign = Primitive.InvNonUniformScaleAndDeterminantSign.w;
+	Intermediates.TangentToWorld = CalcTangentToWorldNoScale(Intermediates.TangentToLocal);
+	
+	// Color
+	{
+		// Used to account for offset due to shifted points bounds
+		float3 OffsetPosition = Intermediates.Position + BoundsOffset;
+
+		half3 Color = half3(0x000000FF & (InstanceData.Color >> 16), 0x000000FF & (InstanceData.Color >> 8), 0x000000FF & InstanceData.Color) * 0.003921568627451;
+		half3 ElevationColor = lerp(ElevationColorBottom, ElevationColorTop, clamp((OffsetPosition.z + BoundsSize.z * 0.5) / BoundsSize.z, 0, 1));
+		half3 PositionColor = (OffsetPosition + BoundsSize * 0.5) / BoundsSize;
+		half3 ClassificationColor = ClassificationColors[0x0000001F & InstanceData.Color];
+
+		// Pick color source
+		Color = lerp(Color, lerp(lerp(PositionColor, ClassificationColor, bUseClassification), ElevationColor, bUseElevationColor), bUseColorOverride);
+
+		// Mix with Intensity
+		{
+			half Intensity = (0x000000FF & (InstanceData.Color >> 24)) * 0.003921568627451;
+			Color = lerp(Color, Color * Intensity, IntensityInfluence);
+		}
+
+
+		// Color Adjustment
+		{
+			Color = lerp(dot(Color, half3(0.299, 0.587, 0.144)), Color, Saturation.rgb * Saturation.a);
+			Color *= Tint;
+			Color = pow(Color, 2.2 * Gamma.rgb * Gamma.a);
+			Color = (Color - 0.5) * Contrast.rgb * Contrast.a + 0.5;
+			Color += Offset.rgb * Offset.a;
+			Color = clamp(Color, 0, 1);
+		}
+
+		Intermediates.Color = half4(lerp(Color, LODColors[min((int)InstanceData.VirtualDepth, 21)], bUseLODColoration), bUseCircle);
+	}
+
+	return Intermediates;
+}
+
+FVertexFactoryInterpolantsVSToPS VertexFactoryGetInterpolantsVSToPS(FVertexFactoryInput Input, FVertexFactoryIntermediates Intermediates, FMaterialVertexParameters VertexParameters)
+{
+	FVertexFactoryInterpolantsVSToPS Interpolants = (FVertexFactoryInterpolantsVSToPS)0;
+	Interpolants.TangentToWorld0 = float4(Intermediates.TangentToWorld[0], 0);
+	Interpolants.TangentToWorld2 = float4(Intermediates.TangentToWorld[2], Intermediates.TangentToWorldSign);
+	Interpolants.Color = VertexParameters.VertexColor;
+	Interpolants.TexCoords = PrecomputedUV_Quad[Input.VertexId % 4 + 4];    // Using inverted V
+
+#if INSTANCED_STEREO
+	Interpolants.EyeIndex = 0;
+#endif
+
+	return Interpolants;
+}
+
+half3x3 VertexFactoryGetTangentToLocal(FVertexFactoryInput Input, FVertexFactoryIntermediates Intermediates) { return Intermediates.TangentToLocal; }
+
+// @return translated world position
+float4 VertexFactoryGetWorldPosition(FVertexFactoryInput Input, FVertexFactoryIntermediates Intermediates)
+{
+	float4 WorldPosition = TransformLocalToTranslatedWorld(Intermediates.Position);
+	WorldPosition.xyz = ProcessPosition(WorldPosition.xyz, Intermediates.Scale);
+	return WorldPosition;
+}
+float4 VertexFactoryGetWorldPosition(FPositionOnlyVertexFactoryInput Input)
+{
+	FPointCloudInstanceData InstanceData = GetInstanceData(Input.VertexId);
+
+	float4 WorldPosition = TransformLocalToTranslatedWorld(InstanceData.Position);
+	WorldPosition.xyz = ProcessPosition(WorldPosition.xyz, InstanceData.Scale);
+	return WorldPosition;
+}
+float4 VertexFactoryGetRasterizedWorldPosition(FVertexFactoryInput Input, FVertexFactoryIntermediates Intermediates, float4 InWorldPosition) { return InWorldPosition; }
+float3 VertexFactoryGetPositionForVertexLighting(FVertexFactoryInput Input, FVertexFactoryIntermediates Intermediates, float3 TranslatedWorldPosition) { return TranslatedWorldPosition; }
+
+float3 VertexFactoryGetWorldNormal(FPositionAndNormalOnlyVertexFactoryInput Input) { return RotateLocalToWorld(Input.Normal.xyz, 0); }
+float4 VertexFactoryGetWorldPosition(FPositionAndNormalOnlyVertexFactoryInput Input)
+{
+	FPointCloudInstanceData InstanceData = GetInstanceData(Input.VertexId);
+
+	float4 WorldPosition = TransformLocalToTranslatedWorld(InstanceData.Position, 0);
+	WorldPosition.xyz = ProcessPosition(WorldPosition.xyz, InstanceData.Scale);
+	return WorldPosition;
+}
+float3 VertexFactoryGetWorldNormal(FVertexFactoryInput Input, FVertexFactoryIntermediates Intermediates) { return Intermediates.TangentToWorld[2]; }
+
+// @return previous translated world position
+float4 VertexFactoryGetPreviousWorldPosition(FVertexFactoryInput Input, FVertexFactoryIntermediates Intermediates)
+{
+	float4x4 PreviousLocalToWorldTranslated = GetPrimitiveData(0).PreviousLocalToWorld;
+    PreviousLocalToWorldTranslated[3][0] += ResolvedView.PrevPreViewTranslation.x;
+    PreviousLocalToWorldTranslated[3][1] += ResolvedView.PrevPreViewTranslation.y;
+    PreviousLocalToWorldTranslated[3][2] += ResolvedView.PrevPreViewTranslation.z;
+    return mul(float4(Intermediates.Position, 1), PreviousLocalToWorldTranslated);
+}
+
+/** Converts from vertex factory specific input to a FMaterialVertexParameters, which is used by vertex shader material inputs. */
+FMaterialVertexParameters GetMaterialVertexParameters(FVertexFactoryInput Input, FVertexFactoryIntermediates Intermediates, float3 WorldPosition, half3x3 TangentToLocal)
+{
+    FMaterialVertexParameters Result = (FMaterialVertexParameters) 0;
+    
+	Result.WorldPosition = WorldPosition;
+    Result.TangentToWorld = Intermediates.TangentToWorld;
+    Result.PreSkinnedNormal = half3(-1, -1, 1);
+    Result.PreSkinnedPosition = WorldPosition;
+    Result.VertexColor = Intermediates.Color;
+
+#if NUM_MATERIAL_TEXCOORDS_VERTEX
+	Result.TexCoords[0] = PrecomputedUV_Quad[Input.VertexId % 4];
+#endif
+
+    return Result;
+}
+
+FMaterialPixelParameters GetMaterialPixelParameters(FVertexFactoryInterpolantsVSToPS Interpolants, float4 SvPosition)
+{
+    FMaterialPixelParameters Result = MakeInitializedMaterialPixelParameters();
+
+    Result.Particle.Color = half4(1, 1, 1, 1);
+    Result.TwoSidedSign = 1;
+	Result.VertexColor = Interpolants.Color;
+	Result.VertexColor.a = lerp(1, 1 - distance(Interpolants.TexCoords, 0.5), Interpolants.Color.a);
+
+    half3 TangentToWorld0 = Interpolants.TangentToWorld0.xyz;
+    half4 TangentToWorld2 = Interpolants.TangentToWorld2;
+    Result.UnMirrored = TangentToWorld2.w;
+    Result.TangentToWorld = AssembleTangentToWorld(TangentToWorld0, TangentToWorld2);
+
+#if NUM_TEX_COORD_INTERPOLATORS
+    Result.TexCoords[0] = Interpolants.TexCoords;
+#endif
+
+    return Result;
+}
+
+#if USE_INSTANCING
+float4 VertexFactoryGetInstanceHitProxyId(FVertexFactoryInput Input, FVertexFactoryIntermediates Intermediates) { return 0; }
+#endif
+
+float4 VertexFactoryGetTranslatedPrimitiveVolumeBounds(FVertexFactoryInterpolantsVSToPS Interpolants)
+{
+	return float4(Primitive.ObjectWorldPositionAndRadius.xyz + ResolvedView.PreViewTranslation.xyz, Primitive.ObjectWorldPositionAndRadius.w);
+}
+
+#if INSTANCED_STEREO
+uint VertexFactoryGetEyeIndex(uint InstanceId)
+{
+	#if USE_INSTANCING
+		return InstancedEyeIndex;
+	#else
+		return InstanceId & 1;
+	#endif
+}
+#endif
+
+#if USING_TESSELLATION
+	struct FVertexFactoryInterpolantsVSToDS
+	{
+		FVertexFactoryInterpolantsVSToPS InterpolantsVSToPS;
+	};
+
+	float2 VertexFactoryGetTextureCoordinateDS( FVertexFactoryInterpolantsVSToDS Interpolants )
+	{
+#if NUM_MATERIAL_TEXCOORDS
+		return Interpolants.InterpolantsVSToPS.TexCoords.xy;
+#else // #if NUM_MATERIAL_TEXCOORDS
+		return float2(0,0);
+#endif // #if NUM_MATERIAL_TEXCOORDS
+	}
+
+	FVertexFactoryInterpolantsVSToPS VertexFactoryAssignInterpolants(FVertexFactoryInterpolantsVSToDS Input)
+	{
+		return Input.InterpolantsVSToPS;
+	}
+
+	/** Converts from vertex factory specific interpolants to a FMaterialTessellationParameters, which is used by material inputs. */
+	FMaterialTessellationParameters GetMaterialTessellationParameters(FVertexFactoryInterpolantsVSToDS Interpolants, float3 CameraLocalWorldPosition)
+	{
+		FMaterialTessellationParameters	Result;
+
+#if NUM_TEX_COORD_INTERPOLATORS
+		Result.TexCoords[0] = Interpolants.InterpolantsVSToPS.TexCoords;
+#endif // NUM_TEX_COORD_INTERPOLATORS
+
+#if INTERPOLATE_VERTEX_COLOR
+		Result.VertexColor = Interpolants.InterpolantsVSToPS.Color;
+#endif	// INTERPOLATE_VERTEX_COLOR
+
+		half3 TangentToWorld0 = Interpolants.InterpolantsVSToPS.TangentToWorld0.xyz;
+		half4 TangentToWorld2 = Interpolants.InterpolantsVSToPS.TangentToWorld2;
+
+		Result.TangentToWorld = AssembleTangentToWorld( TangentToWorld0, TangentToWorld2 );
+
+		Result.TangentToWorldPreScale = 1;
+		Result.WorldPosition = CameraLocalWorldPosition + ResolvedView.WorldCameraOrigin;
+
+		return Result;
+	}
+
+	FVertexFactoryInterpolantsVSToDS VertexFactoryGetInterpolantsVSToDS(FVertexFactoryInput Input, FVertexFactoryIntermediates Intermediates, FMaterialVertexParameters VertexParameters)
+	{
+		FVertexFactoryInterpolantsVSToDS Interpolants;
+		Interpolants.InterpolantsVSToPS = VertexFactoryGetInterpolantsVSToPS(Input, Intermediates, VertexParameters);
+		return Interpolants;
+	}
+
+	FVertexFactoryInterpolantsVSToDS VertexFactoryInterpolate(FVertexFactoryInterpolantsVSToDS a, float aInterp, FVertexFactoryInterpolantsVSToDS b, float bInterp)
+	{
+		FVertexFactoryInterpolantsVSToDS O;
+	
+		// Do we really need to interpolate TangentToWorld2 here? It should be replaced by the
+		// interpolated normal from 'whatever' interpolation scheme we're using
+	
+		TESSELLATION_INTERPOLATE_MEMBER(InterpolantsVSToPS.TangentToWorld0.xyz);
+		TESSELLATION_INTERPOLATE_MEMBER(InterpolantsVSToPS.TangentToWorld2);
+
+#if INTERPOLATE_VERTEX_COLOR
+		TESSELLATION_INTERPOLATE_MEMBER(InterpolantsVSToPS.Color);
+#endif
+
+#if NUM_TEX_COORD_INTERPOLATORS
+		TESSELLATION_INTERPOLATE_MEMBER(InterpolantsVSToPS.TexCoords);
+#endif
+
+		return O;
+	}
+
+	float3x3 VertexFactoryGetTangentToLocalDS(FVertexFactoryInterpolantsVSToDS Interpolants)
+	{
+		// This duplicates stuff already going on in GetMaterialTessellationParameters(), so
+		// maybe the hull shader could leverage that instead?
+		half3 TangentToWorld0 = Interpolants.InterpolantsVSToPS.TangentToWorld0.xyz;
+		half4 TangentToWorld2 = Interpolants.InterpolantsVSToPS.TangentToWorld2;
+		float3x3 TangentToWorld = AssembleTangentToWorld( TangentToWorld0, TangentToWorld2 );
+
+		return TangentToWorld;
+	}
+#endif // #if USING_TESSELLATION
+
+uint VertexFactoryGetPrimitiveId(FVertexFactoryInterpolantsVSToPS Interpolants)
+{
+	return 0;
+}
\ No newline at end of file
diff --git a/Engine/Plugins/Enterprise/LidarPointCloud/Source/LidarPointCloudEditor/Private/LidarPointCloudEditor.cpp b/Engine/Plugins/Enterprise/LidarPointCloud/Source/LidarPointCloudEditor/Private/LidarPointCloudEditor.cpp
new file mode 100644
--- /dev/null
+++ b/Engine/Plugins/Enterprise/LidarPointCloud/Source/LidarPointCloudEditor/Private/LidarPointCloudEditor.cpp
@@ -0,0 +1,748 @@
+// Copyright Epic Games, Inc. All Rights Reserved.
+
+#include "LidarPointCloudEditor.h"
+#include "LidarPointCloudEditorViewport.h"
+#include "LidarPointCloud.h"
+#include "LidarPointCloudShared.h"
+#include "ILidarPointCloudEditorModule.h"
+
+#include "LidarPointCloudEditorCommands.h"
+
+#include "Misc/ScopedSlowTask.h"
+#include "EditorStyleSet.h"
+#include "Framework/MultiBox/MultiBoxBuilder.h"
+#include "SSingleObjectDetailsPanel.h"
+#include "Widgets/Layout/SBorder.h"
+#include "Widgets/Docking/SDockTab.h"
+#include "UObject/Package.h"
+#include "Framework/Commands/UICommandList.h"
+#include "IContentBrowserSingleton.h"
+#include "ContentBrowserModule.h"
+#include "AssetRegistryModule.h"
+#include "Misc/MessageDialog.h"
+
+#define LOCTEXT_NAMESPACE "LidarPointCloudEditor"
+
+const FName PointCloudEditorAppName = FName(TEXT("LidarPointCloudEditorApp"));
+
+const FName FLidarPointCloudEditor::DetailsTabId(TEXT("Details"));
+const FName FLidarPointCloudEditor::ViewportTabId(TEXT("Viewport"));
+
+/////////////////////////////////////////////////////
+// SPointCloudPropertiesTabBody
+
+class SPointCloudPropertiesTabBody : public SSingleObjectDetailsPanel
+{
+public:
+	SLATE_BEGIN_ARGS(SPointCloudPropertiesTabBody) {}
+	SLATE_END_ARGS()
+
+private:
+	// Pointer back to owning sprite editor instance (the keeper of state)
+	TWeakPtr<class FLidarPointCloudEditor> PointCloudEditorPtr;
+public:
+	void Construct(const FArguments& InArgs, TSharedPtr<FLidarPointCloudEditor> InPointCloudEditor)
+	{
+		PointCloudEditorPtr = InPointCloudEditor;
+
+		SSingleObjectDetailsPanel::Construct(SSingleObjectDetailsPanel::FArguments().HostCommandList(InPointCloudEditor->GetToolkitCommands()).HostTabManager(InPointCloudEditor->GetTabManager()), /*bAutomaticallyObserveViaGetObjectToObserve=*/ true, /*bAllowSearch=*/ true);
+	}
+
+	// SSingleObjectDetailsPanel interface
+	virtual UObject* GetObjectToObserve() const override
+	{
+		return PointCloudEditorPtr.Pin()->GetPointCloudBeingEdited();
+	}
+
+	virtual TSharedRef<SWidget> PopulateSlot(TSharedRef<SWidget> PropertyEditorWidget) override
+	{
+		return SNew(SVerticalBox)
+			+ SVerticalBox::Slot()
+			.FillHeight(1)
+			[
+				PropertyEditorWidget
+			];
+	}
+	// End of SSingleObjectDetailsPanel interface
+};
+
+//////////////////////////////////////////////////////////////////////////
+// FLidarPointCloudEditor
+
+FLidarPointCloudEditor::FLidarPointCloudEditor()
+	: PointCloudBeingEdited(nullptr)
+	, bEditMode(false)
+{
+}
+
+FLidarPointCloudEditor::~FLidarPointCloudEditor()
+{
+	// Unregister from the cloud before closing
+	if (PointCloudBeingEdited)
+	{
+		PointCloudBeingEdited->OnPointCloudRebuilt().RemoveAll(this);
+	}
+
+	DeselectPoints();
+}
+
+void FLidarPointCloudEditor::RegisterTabSpawners(const TSharedRef<class FTabManager>& InTabManager)
+{
+	WorkspaceMenuCategory = InTabManager->AddLocalWorkspaceMenuCategory(LOCTEXT("WorkspaceMenu_LidarPointCloudEditor", "LiDAR Point Cloud Editor"));
+	auto WorkspaceMenuCategoryRef = WorkspaceMenuCategory.ToSharedRef();
+
+	FAssetEditorToolkit::RegisterTabSpawners(InTabManager);
+
+	InTabManager->RegisterTabSpawner(FLidarPointCloudEditor::DetailsTabId, FOnSpawnTab::CreateSP(this, &FLidarPointCloudEditor::SpawnTab_Details))
+		.SetDisplayName(LOCTEXT("DetailsTabLabel", "Details"))
+		.SetGroup(WorkspaceMenuCategoryRef)
+		.SetIcon(FSlateIcon(FEditorStyle::GetStyleSetName(), "LevelEditor.Tabs.Details"));
+
+	InTabManager->RegisterTabSpawner(FLidarPointCloudEditor::ViewportTabId, FOnSpawnTab::CreateSP(this, &FLidarPointCloudEditor::SpawnTab_Viewport))
+		.SetDisplayName(LOCTEXT("ViewportTabLabel", "Viewport"))
+		.SetGroup(WorkspaceMenuCategoryRef)
+		.SetIcon(FSlateIcon(FEditorStyle::GetStyleSetName(), "LevelEditor.Tabs.Viewports"));
+}
+
+void FLidarPointCloudEditor::UnregisterTabSpawners(const TSharedRef<class FTabManager>& InTabManager)
+{
+	FAssetEditorToolkit::UnregisterTabSpawners(InTabManager);
+
+	InTabManager->UnregisterTabSpawner(FLidarPointCloudEditor::DetailsTabId);
+	InTabManager->UnregisterTabSpawner(FLidarPointCloudEditor::ViewportTabId);
+}
+
+FName FLidarPointCloudEditor::GetToolkitFName() const { return FName("LidarPointCloudEditor"); }
+FText FLidarPointCloudEditor::GetBaseToolkitName() const { return LOCTEXT("LidarPointCloudEditorAppLabel", "LiDAR Point Cloud Editor"); }
+FText FLidarPointCloudEditor::GetToolkitToolTipText() const { return FAssetEditorToolkit::GetToolTipTextForObject(PointCloudBeingEdited); }
+FLinearColor FLidarPointCloudEditor::GetWorldCentricTabColorScale() const { return FLinearColor::White; }
+FString FLidarPointCloudEditor::GetWorldCentricTabPrefix() const { return TEXT("LidarPointCloudEditor"); }
+
+FText FLidarPointCloudEditor::GetToolkitName() const
+{
+	const bool bDirtyState = PointCloudBeingEdited->GetOutermost()->IsDirty();
+
+	FFormatNamedArguments Args;
+	Args.Add(TEXT("PointCloudName"), FText::FromString(PointCloudBeingEdited->GetName()));
+	Args.Add(TEXT("DirtyState"), bDirtyState ? FText::FromString(TEXT("*")) : FText::GetEmpty());
+	return FText::Format(LOCTEXT("LidarPointCloudEditorToolkitName", "{PointCloudName}{DirtyState}"), Args);
+}
+
+void FLidarPointCloudEditor::AddReferencedObjects(FReferenceCollector& Collector)
+{
+	if (IsValid(PointCloudBeingEdited))
+	{
+		Collector.AddReferencedObject(PointCloudBeingEdited);
+	}
+}
+
+void FLidarPointCloudEditor::InitPointCloudEditor(const EToolkitMode::Type Mode, const TSharedPtr< class IToolkitHost >& InitToolkitHost, class ULidarPointCloud* InitPointCloud)
+{
+	FLidarPointCloudEditorCommands::Register();
+
+	PointCloudBeingEdited = InitPointCloud;
+
+	// Register for rebuilding events
+	PointCloudBeingEdited->OnPointCloudRebuilt().AddSP(this, &FLidarPointCloudEditor::OnPointCloudRebuilt);
+
+	TSharedPtr<FLidarPointCloudEditor> PointCloudEditor = SharedThis(this);
+
+	Viewport = SNew(SLidarPointCloudEditorViewport).PointCloudEditor(SharedThis(this)).ObjectToEdit(PointCloudBeingEdited);
+
+	// Default layout
+	const TSharedRef<FTabManager::FLayout> StandaloneDefaultLayout = FTabManager::NewLayout("Standalone_LidarPointCloudEditor_Layout_v2")
+		->AddArea
+		(
+			FTabManager::NewPrimaryArea()
+			->SetOrientation(Orient_Vertical)
+			->Split
+			(
+				FTabManager::NewStack()
+				->SetSizeCoefficient(0.1f)
+				->SetHideTabWell(true)
+				->AddTab(GetToolbarTabId(), ETabState::OpenedTab)
+			)
+			->Split
+			(
+				FTabManager::NewSplitter()->SetOrientation(Orient_Horizontal)
+				->SetSizeCoefficient(0.9f)
+				->Split
+				(
+					FTabManager::NewStack()
+					->SetSizeCoefficient(0.6f)
+					->AddTab(FLidarPointCloudEditor::ViewportTabId, ETabState::OpenedTab)
+					->SetHideTabWell(true)
+				)
+				->Split
+				(
+					FTabManager::NewStack()
+					->SetSizeCoefficient(0.4f)
+					->AddTab(FLidarPointCloudEditor::DetailsTabId, ETabState::OpenedTab)
+					->SetHideTabWell(true)
+				)
+			)
+		);
+
+	// Initialize the asset editor
+	InitAssetEditor(Mode, InitToolkitHost, PointCloudEditorAppName, StandaloneDefaultLayout, /*bCreateDefaultStandaloneMenu=*/ true, /*bCreateDefaultToolbar=*/ true, InitPointCloud);
+
+	ExtendToolBar();
+
+	BindEditorCommands();
+
+	RegenerateMenusAndToolbars();
+}
+
+void FLidarPointCloudEditor::SelectPointsByFrustum(const FConvexVolume& SelectionFrustum, bool bAdditive)
+{
+	if (bAdditive && SelectedPoints.Num() > 0)
+	{
+		TArray<FLidarPointCloudPoint*> NewSelectedPoints;
+		PointCloudBeingEdited->GetPointsInFrustum(NewSelectedPoints, SelectionFrustum, true);
+
+		// Merge selections and set selection flag for new selection of points
+		for (FLidarPointCloudPoint** Data = NewSelectedPoints.GetData(), **DataEnd = Data + NewSelectedPoints.Num(); Data != DataEnd; ++Data)
+		{
+			if (!(*Data)->bSelected)
+			{
+				(*Data)->bSelected = true;
+				SelectedPoints.Add(*Data);
+			}
+		}
+	}
+	else
+	{
+		// Clear any existing selection
+		DeselectPoints();
+
+		PointCloudBeingEdited->GetPointsInFrustum(SelectedPoints, SelectionFrustum, true);
+
+		// Set selection flag for new selection of points
+		for (FLidarPointCloudPoint** Data = SelectedPoints.GetData(), **DataEnd = Data + SelectedPoints.Num(); Data != DataEnd; ++Data)
+		{
+			(*Data)->bSelected = true;
+		}
+	}	
+}
+
+void FLidarPointCloudEditor::DeselectPointsByFrustum(const FConvexVolume& SelectionFrustum)
+{
+	TArray<FLidarPointCloudPoint*> NewSelectedPoints;
+	PointCloudBeingEdited->GetPointsInFrustum(NewSelectedPoints, SelectionFrustum, true);
+
+	// Unset selection flag for selection of points
+	for (FLidarPointCloudPoint** Data = NewSelectedPoints.GetData(), **DataEnd = Data + NewSelectedPoints.Num(); Data != DataEnd; ++Data)
+	{
+		(*Data)->bSelected = false;
+	}
+
+	// Remove deselected points from the selection list
+	FLidarPointCloudPoint** Data = SelectedPoints.GetData();
+	for (int32 i = 0; i < SelectedPoints.Num(); ++Data, ++i)
+	{
+		if (!(*Data)->bSelected)
+		{
+			SelectedPoints.RemoveAtSwap(i--, 1, false);
+			--Data;
+		}
+	}
+
+	SelectedPoints.Shrink();
+}
+
+void FLidarPointCloudEditor::DeselectPoints()
+{
+	for (FLidarPointCloudPoint** Data = SelectedPoints.GetData(), **DataEnd = Data + SelectedPoints.Num(); Data != DataEnd; ++Data)
+	{
+		(*Data)->bSelected = false;
+	}
+
+	SelectedPoints.Empty();
+}
+
+void FLidarPointCloudEditor::DeletePoints()
+{
+	// Permanently remove points
+	if (FMessageDialog::Open(EAppMsgType::Type::YesNo, FText::FromString("Warning: This operation cannot be reversed!\nAre you sure you want to continue?")) != EAppReturnType::Yes)
+	{
+		return;
+	}
+
+	PointCloudBeingEdited->RemovePoints_NoLock(SelectedPoints, true);
+	PointCloudBeingEdited->RefreshRendering();
+	PointCloudBeingEdited->MarkPackageDirty();
+	SelectedPoints.Empty();
+
+	Viewport->RefreshViewport();
+}
+
+void FLidarPointCloudEditor::DeleteHiddenPoints()
+{
+	// Permanently remove points
+	if (FMessageDialog::Open(EAppMsgType::Type::YesNo, FText::FromString("Warning: This operation cannot be reversed!\nAre you sure you want to continue?")) != EAppReturnType::Yes)
+	{
+		return;
+	}
+
+	PointCloudBeingEdited->RemoveHiddenPoints(true);
+	DeselectPoints();
+	PointCloudBeingEdited->RefreshRendering();
+	PointCloudBeingEdited->MarkPackageDirty();
+
+	Viewport->RefreshViewport();
+}
+
+void FLidarPointCloudEditor::HidePoints()
+{
+	// Hide points
+	for (FLidarPointCloudPoint** Data = SelectedPoints.GetData(), **DataEnd = Data + SelectedPoints.Num(); Data != DataEnd; ++Data)
+	{
+		(*Data)->bVisible = false;
+		(*Data)->bSelected = false;
+	}
+
+	PointCloudBeingEdited->MarkPointVisibilityDirty();
+
+	SelectedPoints.Empty();
+}
+
+void FLidarPointCloudEditor::UnhideAll()
+{
+	if (PointCloudBeingEdited)
+	{
+		PointCloudBeingEdited->UnhideAll();
+	}
+}
+
+TSharedRef<SWidget> FLidarPointCloudEditor::BuildPointCloudStatistics()
+{
+	return SNew(SHorizontalBox)
+		+ SHorizontalBox::Slot()
+		.FillWidth(0.5f)
+		[
+			SNew(SVerticalBox)
+			+ SVerticalBox::Slot()
+			.AutoHeight()
+			.VAlign(VAlign_Center)
+			.Padding(4.0f)
+			[
+				SNew(STextBlock).Text_Lambda([this] { return FText::Format(LOCTEXT("PointCount", "Point Count: {0}"), PointCloudBeingEdited ? PointCloudBeingEdited->GetNumPoints() : 0); })
+			]
+			+ SVerticalBox::Slot()
+			.AutoHeight()
+			.VAlign(VAlign_Center)
+			.Padding(4.0f)
+			[
+				SNew(STextBlock).Text_Lambda([this] { return FText::Format(LOCTEXT("NodeCount", "Node Count: {0}"), PointCloudBeingEdited ? PointCloudBeingEdited->GetNumNodes() : 0); })
+			]
+			+ SVerticalBox::Slot()
+			.AutoHeight()
+			.VAlign(VAlign_Center)
+			.Padding(4.0f)
+			[
+				SNew(STextBlock).Text_Lambda([this] { return FText::Format(LOCTEXT("Collider", "Collider: {0} poly"), PointCloudBeingEdited && PointCloudBeingEdited->HasCollisionData() ? PointCloudBeingEdited->Octree.GetCollisionData()->Indices.Num() : 0); })
+			]
+		]
+
+		+ SHorizontalBox::Slot()
+		.FillWidth(0.5f)
+		[
+			SNew(SVerticalBox)
+			+ SVerticalBox::Slot()
+			.AutoHeight()
+			.VAlign(VAlign_Center)
+			.Padding(4.0f)
+			[
+				SNew(STextBlock).Text_Lambda([this]
+				{
+					FVector BoundingSize = PointCloudBeingEdited ? PointCloudBeingEdited->GetPointsBounds().GetSize() : FVector::ZeroVector;
+					return FText::Format(LOCTEXT("PCBounds", "Bounds: {0} x {1} x {2}"), FMath::CeilToInt(BoundingSize.X), FMath::CeilToInt(BoundingSize.Y), FMath::CeilToInt(BoundingSize.Z));
+				})
+			]
+			+ SVerticalBox::Slot()
+			.AutoHeight()
+			.VAlign(VAlign_Center)
+			.Padding(4.0f)
+			[
+				SNew(STextBlock).Text_Lambda([this] { return FText::Format(LOCTEXT("PCSize", "Data Size: {0} MB"), PointCloudBeingEdited ? PointCloudBeingEdited->GetDataSize() : 0); })
+			]
+		];
+}
+
+TSharedRef<SWidget> FLidarPointCloudEditor::GenerateExtractionMenuContent()
+{
+	FMenuBuilder MenuBuilder(true, Viewport->GetCommandList());
+
+	MenuBuilder.BeginSection("ExtractSelection", LOCTEXT("ExtractSelection", "Extract Selection"));
+	{
+		MenuBuilder.AddMenuEntry(FLidarPointCloudEditorCommands::Get().Extract);
+		MenuBuilder.AddMenuEntry(FLidarPointCloudEditorCommands::Get().ExtractCopy);
+	}
+	MenuBuilder.EndSection();
+
+	return MenuBuilder.MakeWidget();
+}
+
+TSharedRef<SWidget> FLidarPointCloudEditor::GenerateCollisionMenuContent()
+{
+	FMenuBuilder MenuBuilder(true, Viewport->GetCommandList());
+
+	MenuBuilder.BeginSection("Collision", LOCTEXT("Collision", "Collision"));
+	{
+		MenuBuilder.AddMenuEntry(FLidarPointCloudEditorCommands::Get().BuildCollision);
+		MenuBuilder.AddMenuEntry(FLidarPointCloudEditorCommands::Get().RemoveCollision);
+	}
+	MenuBuilder.EndSection();
+
+	return MenuBuilder.MakeWidget();
+}
+
+TSharedRef<SWidget> FLidarPointCloudEditor::GenerateDeleteMenuContent()
+{
+	FMenuBuilder MenuBuilder(true, Viewport->GetCommandList());
+
+	MenuBuilder.BeginSection("Delete", LOCTEXT("Delete", "Delete"));
+	{
+		MenuBuilder.AddMenuEntry(FLidarPointCloudEditorCommands::Get().DeleteSelected);
+		MenuBuilder.AddMenuEntry(FLidarPointCloudEditorCommands::Get().DeleteHidden, NAME_None, TAttribute<FText>(), TAttribute<FText>(), FSlateIcon(FLidarPointCloudStyle::GetStyleSetName(), "LidarPointCloudEditor.DeleteSelected"));
+	}
+	MenuBuilder.EndSection();
+
+	return MenuBuilder.MakeWidget();
+}
+
+void FLidarPointCloudEditor::ExtendToolBar()
+{
+	struct Local
+	{
+		static void FillToolbar(FToolBarBuilder& ToolbarBuilder, FLidarPointCloudEditor* ThisEditor)
+		{
+			auto Commands = &FLidarPointCloudEditorCommands::Get();
+
+			ToolbarBuilder.BeginSection("Camera");
+			{
+				ToolbarBuilder.AddToolBarButton(Commands->ResetCamera, NAME_None, TAttribute<FText>(), TAttribute<FText>(), FSlateIcon(FEditorStyle::GetStyleSetName(), "StaticMeshEditor.ResetCamera"));
+			}
+			ToolbarBuilder.EndSection();
+
+			ToolbarBuilder.BeginSection("LidarPointCloud");
+			{
+				ToolbarBuilder.AddToolBarButton(Commands->Center, NAME_None, TAttribute<FText>(), TAttribute<FText>(), FSlateIcon(FEditorStyle::GetStyleSetName(), "AnimViewportMenu.TranslateMode"));
+				ToolbarBuilder.AddToolBarButton(Commands->Merge);
+				ToolbarBuilder.AddToolBarButton(Commands->Align);
+				ToolbarBuilder.AddComboButton(
+					FUIAction(FExecuteAction(), FCanExecuteAction()),
+					FOnGetContent::CreateSP(ThisEditor, &FLidarPointCloudEditor::GenerateCollisionMenuContent),
+					LOCTEXT("Collision_Label", "Collision"),
+					LOCTEXT("Collision_Tooltip", "Collision options"),
+					FSlateIcon(FLidarPointCloudStyle::GetStyleSetName(), "LidarPointCloudEditor.BuildCollision")
+				);
+			}
+			ToolbarBuilder.EndSection();
+
+			ToolbarBuilder.BeginSection("LidarPointCloudEdit");
+			{
+				ToolbarBuilder.AddToolBarButton(Commands->EditMode);
+				ToolbarBuilder.AddToolBarButton(Commands->HideSelected);
+				ToolbarBuilder.AddToolBarButton(Commands->UnhideAll);
+				ToolbarBuilder.AddComboButton(
+					FUIAction(FExecuteAction(), FCanExecuteAction()),
+					FOnGetContent::CreateSP(ThisEditor, &FLidarPointCloudEditor::GenerateDeleteMenuContent),
+					LOCTEXT("Delete_Label", "Delete"),
+					LOCTEXT("Delete_Tooltip", "Point deletion options"),
+					FSlateIcon(FLidarPointCloudStyle::GetStyleSetName(), "LidarPointCloudEditor.DeleteSelected")
+				);
+				ToolbarBuilder.AddComboButton(
+					FUIAction(FExecuteAction(), FCanExecuteAction::CreateSP(ThisEditor, &FLidarPointCloudEditor::HasSelectedPoints)),
+					FOnGetContent::CreateSP(ThisEditor, &FLidarPointCloudEditor::GenerateExtractionMenuContent),
+					LOCTEXT("Extract_Label", "Extract"),
+					LOCTEXT("Extract_Tooltip", "Selection extraction options"),
+					FSlateIcon(FLidarPointCloudStyle::GetStyleSetName(), "LidarPointCloudEditor.Extract")
+				);
+			}
+			ToolbarBuilder.EndSection();
+		}
+	};
+
+	TSharedPtr<FExtender> ToolbarExtender = MakeShareable(new FExtender);
+	ToolbarExtender->AddToolBarExtension("Asset", EExtensionHook::After, Viewport->GetCommandList(), FToolBarExtensionDelegate::CreateStatic(&Local::FillToolbar, this));
+	AddToolbarExtender(ToolbarExtender);
+
+	AddToolbarExtender(ILidarPointCloudEditorModule::Get().GetToolBarExtensibilityManager()->GetAllExtenders(GetToolkitCommands(), GetEditingObjects()));
+}
+
+void FLidarPointCloudEditor::BindEditorCommands()
+{
+	TSharedPtr<FUICommandList> CommandList = Viewport->GetCommandList();
+	
+	const FLidarPointCloudEditorCommands& Commands = FLidarPointCloudEditorCommands::Get();
+
+	CommandList->MapAction(Commands.Center, FExecuteAction::CreateSP(this, &FLidarPointCloudEditor::ToggleCenter), FCanExecuteAction(), FIsActionChecked::CreateSP(this, &FLidarPointCloudEditor::IsCentered));
+	CommandList->MapAction(Commands.BuildCollision, FExecuteAction::CreateSP(this, &FLidarPointCloudEditor::BuildCollision));
+	CommandList->MapAction(Commands.RemoveCollision, FExecuteAction::CreateSP(this, &FLidarPointCloudEditor::RemoveCollision), FCanExecuteAction::CreateSP(this, &FLidarPointCloudEditor::HasCollisionData));
+	CommandList->MapAction(Commands.EditMode, FExecuteAction::CreateSP(this, &FLidarPointCloudEditor::ToggleEditMode), FCanExecuteAction(), FIsActionChecked::CreateSP(this, &FLidarPointCloudEditor::IsEditMode));
+	CommandList->MapAction(Commands.UnhideAll, FExecuteAction::CreateSP(this, &FLidarPointCloudEditor::UnhideAll));
+	CommandList->MapAction(Commands.HideSelected, FExecuteAction::CreateSP(this, &FLidarPointCloudEditor::HidePoints), FCanExecuteAction::CreateSP(this, &FLidarPointCloudEditor::HasSelectedPoints));
+	CommandList->MapAction(Commands.DeleteSelected, FExecuteAction::CreateSP(this, &FLidarPointCloudEditor::DeletePoints), FCanExecuteAction::CreateSP(this, &FLidarPointCloudEditor::HasSelectedPoints));
+	CommandList->MapAction(Commands.DeleteHidden, FExecuteAction::CreateSP(this, &FLidarPointCloudEditor::DeleteHiddenPoints));
+	CommandList->MapAction(Commands.Extract, FExecuteAction::CreateSP(this, &FLidarPointCloudEditor::Extract), FCanExecuteAction::CreateSP(this, &FLidarPointCloudEditor::HasSelectedPoints));
+	CommandList->MapAction(Commands.ExtractCopy, FExecuteAction::CreateSP(this, &FLidarPointCloudEditor::ExtractCopy), FCanExecuteAction::CreateSP(this, &FLidarPointCloudEditor::HasSelectedPoints));
+	CommandList->MapAction(Commands.Merge, FExecuteAction::CreateSP(this, &FLidarPointCloudEditor::Merge));
+	CommandList->MapAction(Commands.Align, FExecuteAction::CreateSP(this, &FLidarPointCloudEditor::Align));
+}
+
+TSharedRef<SDockTab> FLidarPointCloudEditor::SpawnTab_Details(const FSpawnTabArgs& Args)
+{
+	TSharedPtr<FLidarPointCloudEditor> PointCloudEditorPtr = SharedThis(this);
+
+	// Spawn the tab
+	return SNew(SDockTab)
+		.Label(LOCTEXT("DetailsTab_Title", "Details"))
+		[
+			SNew(SVerticalBox)
+				+ SVerticalBox::Slot()
+				.AutoHeight()
+				.Padding(2.0f)
+				[
+					SNew(SBorder)
+					[
+						BuildPointCloudStatistics()
+					]
+				]
+
+				+ SVerticalBox::Slot()
+				.FillHeight(1.0f)
+				.Padding(2.0f)
+				[
+					SNew(SBorder)
+					.Padding(4.0f)
+					[
+						SNew(SPointCloudPropertiesTabBody, PointCloudEditorPtr)
+					]
+				]
+		];
+}
+
+TSharedRef<SDockTab> FLidarPointCloudEditor::SpawnTab_Viewport(const FSpawnTabArgs& Args)
+{
+	check(Args.GetTabId() == ViewportTabId);
+
+	TSharedRef<SDockTab> SpawnedTab =
+		SNew(SDockTab)
+		.Label(LOCTEXT("ViewportTab_Title", "Viewport"))
+		[
+			Viewport.ToSharedRef()
+		];
+
+	Viewport->SetParentTab(SpawnedTab);
+
+	return SpawnedTab;
+}
+
+void FLidarPointCloudEditor::OnPointCloudRebuilt()
+{
+	// If the cloud asset has been rebuilt, invalidate the selection array
+	SelectedPoints.Empty();
+}
+
+void FLidarPointCloudEditor::Extract()
+{
+	// Skip, if no points are selected
+	if (SelectedPoints.Num() == 0)
+	{
+		return;
+	}
+
+	ULidarPointCloud* NewPointCloud = CreateNewAsset();
+	if (NewPointCloud)
+	{
+		NewPointCloud->SetData(SelectedPoints);
+		NewPointCloud->OriginalCoordinates = PointCloudBeingEdited->OriginalCoordinates;
+		DeletePoints();
+	}
+}
+
+void FLidarPointCloudEditor::ExtractCopy()
+{
+	// Skip, if no points are selected
+	if (SelectedPoints.Num() == 0)
+	{
+		return;
+	}
+
+	ULidarPointCloud* NewPointCloud = CreateNewAsset();
+	if (NewPointCloud)
+	{
+		NewPointCloud->SetData(SelectedPoints);
+		NewPointCloud->OriginalCoordinates = PointCloudBeingEdited->OriginalCoordinates;
+		DeselectPoints();
+	}
+}
+
+void FLidarPointCloudEditor::ToggleCenter()
+{
+	if (IsCentered())
+	{
+		PointCloudBeingEdited->RestoreOriginalCoordinates();
+	}
+	else
+	{
+		PointCloudBeingEdited->CenterPoints();
+	}
+
+	if (Viewport.IsValid())
+	{
+		Viewport->ResetCamera();
+	}
+}
+
+void FLidarPointCloudEditor::ToggleEditMode()
+{
+	bEditMode = !bEditMode;
+
+	if (!bEditMode)
+	{
+		DeselectPoints();
+		Viewport->RefreshViewport();
+	}
+}
+
+bool FLidarPointCloudEditor::IsCentered() const
+{
+	return PointCloudBeingEdited && PointCloudBeingEdited->IsCentered();
+}
+
+void FLidarPointCloudEditor::Merge()
+{
+	TArray<FAssetData> SelectedAssets = SelectAssets(LOCTEXT("SelectMergeSource", "Select Merge Source"));
+	if (SelectedAssets.Num())
+	{
+		TArray<ULidarPointCloud*> Assets;
+		TArray<FString> Names({"Initializing", "Self"});
+
+		for (int32 i = 0; i < SelectedAssets.Num(); ++i)
+		{
+			ULidarPointCloud* Asset = Cast<ULidarPointCloud>(SelectedAssets[i].GetAsset());
+
+			// Skip invalid and self
+			if (!Asset || Asset == PointCloudBeingEdited)
+			{
+				continue;
+			}
+
+			Assets.Add(Asset);
+			Names.Add(SelectedAssets[i].AssetName.ToString());
+		}
+
+		FScopedSlowTask ProgressDialog(SelectedAssets.Num() + 2, LOCTEXT("Merge", "Merging Point Clouds..."));
+		ProgressDialog.MakeDialog();
+		int32 i = 0;
+
+		PointCloudBeingEdited->Merge(Assets, [&ProgressDialog, &i, &Names](float Progress){
+			ProgressDialog.EnterProgressFrame(1.f, FText::FromString(Names[i++]));
+		});
+	}
+}
+
+void FLidarPointCloudEditor::BuildCollision()
+{
+	if (PointCloudBeingEdited)
+	{
+		PointCloudBeingEdited->BuildCollision();
+	}
+}
+
+void FLidarPointCloudEditor::RemoveCollision()
+{
+	if (PointCloudBeingEdited)
+	{
+		PointCloudBeingEdited->RemoveCollision();
+	}
+}
+
+void FLidarPointCloudEditor::Align()
+{
+	TArray<FAssetData> SelectedAssets = SelectAssets(LOCTEXT("SelectAlignmentSources", "Select Alignment Sources"));
+	if (SelectedAssets.Num())
+	{
+		TArray<ULidarPointCloud*> Assets;
+		Assets.Add(PointCloudBeingEdited);
+
+		for (int32 i = 0; i < SelectedAssets.Num(); ++i)
+		{
+			ULidarPointCloud* Asset = Cast<ULidarPointCloud>(SelectedAssets[i].GetAsset());
+
+			// Skip invalid and self
+			if (!Asset || Asset == PointCloudBeingEdited)
+			{
+				continue;
+			}
+
+			Assets.Add(Asset);
+		}
+
+		FScopedSlowTask ProgressDialog(1, LOCTEXT("Align", "Aligning Point Clouds..."));
+		ProgressDialog.MakeDialog();
+		ProgressDialog.EnterProgressFrame(1.f);
+		ULidarPointCloud::AlignClouds(Assets);
+	}
+}
+
+bool FLidarPointCloudEditor::HasCollisionData() const
+{
+	return PointCloudBeingEdited && PointCloudBeingEdited->HasCollisionData();
+}
+
+TArray<FAssetData> FLidarPointCloudEditor::SelectAssets(const FText& Title)
+{
+	// Initialize SaveAssetDialog config
+	FOpenAssetDialogConfig OpenAssetDialogConfig;
+	OpenAssetDialogConfig.DialogTitleOverride = Title;
+	OpenAssetDialogConfig.DefaultPath = "/Game";
+	OpenAssetDialogConfig.bAllowMultipleSelection = true;
+	OpenAssetDialogConfig.AssetClassNames.Emplace(*ULidarPointCloud::StaticClass()->GetName());
+
+	FContentBrowserModule& ContentBrowserModule = FModuleManager::LoadModuleChecked<FContentBrowserModule>("ContentBrowser");
+	return ContentBrowserModule.Get().CreateModalOpenAssetDialog(OpenAssetDialogConfig);
+}
+
+FString FLidarPointCloudEditor::GetSaveAsLocation()
+{
+	// Initialize SaveAssetDialog config
+	FSaveAssetDialogConfig SaveAssetDialogConfig;
+	SaveAssetDialogConfig.DialogTitleOverride = LOCTEXT("SelectExtractDestination", "Select Extract Destination");
+	SaveAssetDialogConfig.DefaultPath = "/Game";
+	SaveAssetDialogConfig.AssetClassNames.Emplace(*ULidarPointCloud::StaticClass()->GetName());
+	SaveAssetDialogConfig.ExistingAssetPolicy = ESaveAssetDialogExistingAssetPolicy::AllowButWarn;
+
+	FContentBrowserModule& ContentBrowserModule = FModuleManager::LoadModuleChecked<FContentBrowserModule>("ContentBrowser");
+	return ContentBrowserModule.Get().CreateModalSaveAssetDialog(SaveAssetDialogConfig);
+}
+
+ULidarPointCloud* FLidarPointCloudEditor::CreateNewAsset()
+{
+	ULidarPointCloud* NewPointCloud = nullptr;
+
+	FString SaveObjectPath = GetSaveAsLocation();
+	if (!SaveObjectPath.IsEmpty())
+	{
+		// Attempt to load existing asset first
+		NewPointCloud = FindObject<ULidarPointCloud>(nullptr, *SaveObjectPath);
+
+		// Proceed to creating a new asset, if needed
+		if (!NewPointCloud)
+		{
+			const FString PackageName = FPackageName::ObjectPathToPackageName(SaveObjectPath);
+			const FString ObjectName = FPackageName::ObjectPathToObjectName(SaveObjectPath);
+
+			NewPointCloud = NewObject<ULidarPointCloud>(CreatePackage(nullptr, *PackageName), ULidarPointCloud::StaticClass(), FName(*ObjectName), EObjectFlags::RF_Public | EObjectFlags::RF_Standalone);
+
+			FAssetRegistryModule::AssetCreated(NewPointCloud);
+			NewPointCloud->MarkPackageDirty();
+		}		
+	}
+
+	return NewPointCloud;
+}
+
+#undef LOCTEXT_NAMESPACE
\ No newline at end of file
diff --git a/Engine/Plugins/Enterprise/LidarPointCloud/Source/LidarPointCloudEditor/Private/LidarPointCloudEditor.h b/Engine/Plugins/Enterprise/LidarPointCloud/Source/LidarPointCloudEditor/Private/LidarPointCloudEditor.h
new file mode 100644
--- /dev/null
+++ b/Engine/Plugins/Enterprise/LidarPointCloud/Source/LidarPointCloudEditor/Private/LidarPointCloudEditor.h
@@ -0,0 +1,108 @@
+// Copyright Epic Games, Inc. All Rights Reserved.
+
+#pragma once
+
+#include "CoreMinimal.h"
+#include "UObject/GCObject.h"
+#include "Toolkits/IToolkitHost.h"
+#include "Toolkits/AssetEditorToolkit.h"
+#include "Widgets/Text/STextBlock.h"
+#include "AssetData.h"
+#include "ConvexVolume.h"
+
+class ULidarPointCloud;
+struct FLidarPointCloudPoint;
+class SWidget;
+
+class FLidarPointCloudEditor : public FAssetEditorToolkit, public FGCObject
+{
+private:
+	ULidarPointCloud* PointCloudBeingEdited;
+
+	TArray<FLidarPointCloudPoint*> SelectedPoints;
+	
+	bool bEditMode;
+
+	/** Preview Viewport widget */
+	TSharedPtr<class SLidarPointCloudEditorViewport> Viewport;
+
+public:
+	FLidarPointCloudEditor();
+	~FLidarPointCloudEditor();
+
+	// IToolkit interface
+	virtual void RegisterTabSpawners(const TSharedRef<class FTabManager>& InTabManager) override;
+	virtual void UnregisterTabSpawners(const TSharedRef<class FTabManager>& InTabManager) override;
+	// End of IToolkit interface
+
+	// FAssetEditorToolkit
+	virtual FName GetToolkitFName() const override;
+	virtual FText GetBaseToolkitName() const override;
+	virtual FText GetToolkitName() const override;
+	virtual FText GetToolkitToolTipText() const override;
+	virtual FLinearColor GetWorldCentricTabColorScale() const override;
+	virtual FString GetWorldCentricTabPrefix() const override;
+
+	virtual void OnToolkitHostingStarted(const TSharedRef<class IToolkit>& Toolkit) override {}
+	virtual void OnToolkitHostingFinished(const TSharedRef<class IToolkit>& Toolkit) override {}
+	// End of FAssetEditorToolkit
+
+	// FSerializableObject interface
+	virtual void AddReferencedObjects(FReferenceCollector& Collector) override;
+	// End of FSerializableObject interface
+
+	void InitPointCloudEditor(const EToolkitMode::Type Mode, const TSharedPtr< class IToolkitHost >& InitToolkitHost, class ULidarPointCloud* InitPointCloud);
+
+	ULidarPointCloud* GetPointCloudBeingEdited() const { return PointCloudBeingEdited; }
+
+	void SelectPointsByFrustum(const FConvexVolume& SelectionFrustum, bool bAdditive);
+	void DeselectPointsByFrustum(const FConvexVolume& SelectionFrustum);
+	void DeselectPoints();
+	void DeletePoints();
+	void DeleteHiddenPoints();
+	void HidePoints();
+	void UnhideAll();
+	TArray<FLidarPointCloudPoint*>& GetSelectedPoints() { return SelectedPoints; }
+
+	bool HasSelectedPoints() const { return SelectedPoints.Num() > 0; }
+
+	bool IsEditMode() const { return bEditMode; }
+
+private:
+	TSharedRef<SWidget> BuildPointCloudStatistics();
+
+	/** Builds the Point Cloud Editor toolbar. */
+	void ExtendToolBar();
+
+	void BindEditorCommands();
+
+	TSharedRef<SDockTab> SpawnTab_Details(const FSpawnTabArgs& Args);
+	TSharedRef<SDockTab> SpawnTab_Viewport(const FSpawnTabArgs& Args);
+
+	/** Called when the owned cloud is rebuilt */
+	void OnPointCloudRebuilt();
+
+	void Extract();
+	void ExtractCopy();
+	bool IsCentered() const;
+	void ToggleCenter();
+	void ToggleEditMode();
+	void Merge();
+	void BuildCollision();
+	void RemoveCollision();
+	void Align();
+	bool HasCollisionData() const;
+
+	TSharedRef<SWidget> GenerateExtractionMenuContent();
+	TSharedRef<SWidget> GenerateCollisionMenuContent();
+	TSharedRef<SWidget> GenerateDeleteMenuContent();
+
+	TArray<FAssetData> SelectAssets(const FText& Title);
+	FString GetSaveAsLocation();
+	ULidarPointCloud* CreateNewAsset();
+
+private:
+	/**	The tab ids for all the tabs used */
+	static const FName DetailsTabId;
+	static const FName ViewportTabId;
+};
\ No newline at end of file
diff --git a/Engine/Plugins/Enterprise/LidarPointCloud/Source/LidarPointCloudEditor/Private/LidarPointCloudEditorCommands.cpp b/Engine/Plugins/Enterprise/LidarPointCloud/Source/LidarPointCloudEditor/Private/LidarPointCloudEditorCommands.cpp
new file mode 100644
--- /dev/null
+++ b/Engine/Plugins/Enterprise/LidarPointCloud/Source/LidarPointCloudEditor/Private/LidarPointCloudEditorCommands.cpp
@@ -0,0 +1,38 @@
+// Copyright Epic Games, Inc. All Rights Reserved.
+
+#include "LidarPointCloudEditorCommands.h"
+
+#define LOCTEXT_NAMESPACE "LidarPointCloudEditor"
+
+//////////////////////////////////////////////////////////////////////////
+// FLidarPointCloudEditorCommands
+
+void FLidarPointCloudEditorCommands::RegisterCommands()
+{
+	UI_COMMAND(SetShowGrid, "Grid", "Displays the viewport grid.", EUserInterfaceActionType::ToggleButton, FInputChord());
+	UI_COMMAND(SetShowBounds, "Bounds", "Toggles display of the bounds of the point cloud.", EUserInterfaceActionType::ToggleButton, FInputChord());
+	UI_COMMAND(SetShowCollision, "Collision", "Toggles display of the collision of the point cloud.", EUserInterfaceActionType::ToggleButton, FInputChord());
+	UI_COMMAND(SetShowNodes, "Nodes", "Toggles display of the nodes of the point cloud.", EUserInterfaceActionType::ToggleButton, FInputChord());
+
+	UI_COMMAND(ResetCamera, "Reset Camera", "Resets the camera to focus on the point cloud.", EUserInterfaceActionType::Button, FInputChord());
+
+	UI_COMMAND(Center, "Center", "Enable, to center the point cloud asset\nDisable, to use original coordinates.", EUserInterfaceActionType::ToggleButton, FInputChord());
+	UI_COMMAND(BuildCollision, "Build Collision", "Builds collision for this point cloud.", EUserInterfaceActionType::Button, FInputChord());
+	UI_COMMAND(RemoveCollision, "Remove Collision", "Removes collision from this point cloud.", EUserInterfaceActionType::Button, FInputChord());
+
+	UI_COMMAND(EditMode, "Edit Mode", "Enables editing of the point cloud.", EUserInterfaceActionType::ToggleButton, FInputChord());
+	UI_COMMAND(HideSelected, "Hide Selected", "Hide selected points.", EUserInterfaceActionType::Button, FInputChord());
+	UI_COMMAND(UnhideAll, "Unhide All", "Resets the visibility of all points in the point cloud.", EUserInterfaceActionType::Button, FInputChord());
+
+	UI_COMMAND(DeleteSelected, "Delete Selected", "Permanently remove selected points from the point cloud.", EUserInterfaceActionType::Button, FInputChord());
+	UI_COMMAND(DeleteHidden, "Delete Hidden", "Permanently remove hidden points from the point cloud.", EUserInterfaceActionType::Button, FInputChord());
+
+	UI_COMMAND(Extract, "Extract", "Extracts the selected points as a separate point cloud asset.", EUserInterfaceActionType::Button, FInputChord());
+	UI_COMMAND(ExtractCopy, "Extract as Copy", "Duplicates the selected points as a separate point cloud asset.", EUserInterfaceActionType::Button, FInputChord());
+	UI_COMMAND(Merge, "Merge", "Merges selected point cloud assets with this one.", EUserInterfaceActionType::Button, FInputChord());
+	UI_COMMAND(Align, "Align", "Aligns selected point cloud assets with this one while retaining overall centering.", EUserInterfaceActionType::Button, FInputChord());
+}
+
+//////////////////////////////////////////////////////////////////////////
+
+#undef LOCTEXT_NAMESPACE
diff --git a/Engine/Plugins/Enterprise/LidarPointCloud/Source/LidarPointCloudEditor/Private/LidarPointCloudEditorCommands.h b/Engine/Plugins/Enterprise/LidarPointCloud/Source/LidarPointCloudEditor/Private/LidarPointCloudEditorCommands.h
new file mode 100644
--- /dev/null
+++ b/Engine/Plugins/Enterprise/LidarPointCloud/Source/LidarPointCloudEditor/Private/LidarPointCloudEditorCommands.h
@@ -0,0 +1,50 @@
+// Copyright Epic Games, Inc. All Rights Reserved.
+
+#pragma once
+
+#include "CoreMinimal.h"
+#include "Framework/Commands/Commands.h"
+#include "EditorStyleSet.h"
+#include "LidarPointCloudStyle.h"
+
+class FLidarPointCloudEditorCommands : public TCommands<FLidarPointCloudEditorCommands>
+{
+public:
+	FLidarPointCloudEditorCommands()
+		: TCommands<FLidarPointCloudEditorCommands>(
+			TEXT("LidarPointCloudEditor"), // Context name for fast lookup
+			NSLOCTEXT("Contexts", "LidarPointCloudEditor", "LiDAR Point Cloud Editor"), // Localized context name for displaying
+			NAME_None, // Parent
+			FLidarPointCloudStyle::GetStyleSetName() // Icon Style Set
+			)
+	{
+	}
+
+	// TCommand<> interface
+	virtual void RegisterCommands() override;
+	// End of TCommand<> interface
+
+public:
+	TSharedPtr< FUICommandInfo > SetShowGrid;
+	TSharedPtr< FUICommandInfo > SetShowBounds;
+	TSharedPtr< FUICommandInfo > SetShowCollision;
+	TSharedPtr< FUICommandInfo > SetShowNodes;
+
+	TSharedPtr< FUICommandInfo > ResetCamera;
+
+	TSharedPtr< FUICommandInfo > Center;
+	TSharedPtr< FUICommandInfo > BuildCollision;
+	TSharedPtr< FUICommandInfo > RemoveCollision;
+
+	TSharedPtr< FUICommandInfo > EditMode;
+	TSharedPtr< FUICommandInfo > HideSelected;
+	TSharedPtr< FUICommandInfo > UnhideAll;
+
+	TSharedPtr< FUICommandInfo > DeleteSelected;
+	TSharedPtr< FUICommandInfo > DeleteHidden;
+
+	TSharedPtr< FUICommandInfo > Extract;
+	TSharedPtr< FUICommandInfo > ExtractCopy;
+	TSharedPtr< FUICommandInfo > Merge;
+	TSharedPtr< FUICommandInfo > Align;
+};
\ No newline at end of file
diff --git a/Engine/Plugins/Enterprise/LidarPointCloud/Source/LidarPointCloudEditor/Private/LidarPointCloudEditorViewportClient.cpp b/Engine/Plugins/Enterprise/LidarPointCloud/Source/LidarPointCloudEditor/Private/LidarPointCloudEditorViewportClient.cpp
new file mode 100644
--- /dev/null
+++ b/Engine/Plugins/Enterprise/LidarPointCloud/Source/LidarPointCloudEditor/Private/LidarPointCloudEditorViewportClient.cpp
@@ -0,0 +1,470 @@
+// Copyright Epic Games, Inc. All Rights Reserved.
+
+#include "LidarPointCloudEditorViewportClient.h"
+#include "LidarPointCloudEditor.h"
+#include "LidarPointCloudEditorViewport.h"
+#include "LidarPointCloud.h"
+#include "LidarPointCloudShared.h"
+#include "LidarPointCloudComponent.h"
+
+#include "ConvexVolume.h"
+#include "CanvasTypes.h"
+#include "Settings/LevelEditorViewportSettings.h"
+#include "Editor/EditorPerProjectUserSettings.h"
+#include "Classes/EditorStyleSettings.h"
+#include "AssetViewerSettings.h"
+
+#define LOCTEXT_NAMESPACE "FLidarPointCloudEditorViewportClient"
+
+namespace {
+	static const float LightRotSpeed = 0.22f;
+	static const float StaticMeshEditor_RotateSpeed = 0.01f;
+	static const float	StaticMeshEditor_TranslateSpeed = 0.25f;
+	static const float GridSize = 2048.0f;
+	static const int32 CellSize = 16;
+	static const float AutoViewportOrbitCameraTranslate = 256.0f;
+
+	static float AmbientCubemapIntensity = 0.4f;
+}
+
+FLidarPointCloudEditorViewportClient::FLidarPointCloudEditorViewportClient(TWeakPtr<FLidarPointCloudEditor> InPointCloudEditor, const TSharedRef<SLidarPointCloudEditorViewport>& InPointCloudEditorViewport, const TSharedRef<FAdvancedPreviewScene>& InPreviewScene, ULidarPointCloud* InPreviewPointCLoud, ULidarPointCloudComponent* InPreviewPointCloudComponent)
+	: FEditorViewportClient(nullptr, &InPreviewScene.Get(), StaticCastSharedRef<SEditorViewport>(InPointCloudEditorViewport))
+	, PointCloudEditorPtr(InPointCloudEditor)
+	, PointCloudEditorViewportPtr(InPointCloudEditorViewport)
+	, SelectionMode(ELidarPointCloudSelectionMode::None)
+{
+	// Setup defaults for the common draw helper.
+	DrawHelper.bDrawPivot = false;
+	DrawHelper.bDrawWorldBox = false;
+	DrawHelper.bDrawKillZ = false;
+	DrawHelper.bDrawGrid = true;
+	DrawHelper.GridColorAxis = FColor(160, 160, 160);
+	DrawHelper.GridColorMajor = FColor(144, 144, 144);
+	DrawHelper.GridColorMinor = FColor(128, 128, 128);
+	DrawHelper.PerspectiveGridSize = GridSize;
+	DrawHelper.NumCells = DrawHelper.PerspectiveGridSize / (CellSize * 2);
+
+	SetViewMode(VMI_Lit);
+
+	EngineShowFlags.SetSeparateTranslucency(true);
+	EngineShowFlags.SetSnap(0);
+	EngineShowFlags.SetCompositeEditorPrimitives(true);
+	OverrideNearClipPlane(1.0f);
+	bUsingOrbitCamera = true;
+
+	AdvancedPreviewScene = static_cast<FAdvancedPreviewScene*>(PreviewScene);
+
+	SetPreviewCloud(InPreviewPointCLoud, InPreviewPointCloudComponent, true);
+
+	// Register delegate to update the show flags when the post processing is turned on or off
+	UAssetViewerSettings::Get()->OnAssetViewerSettingsChanged().AddRaw(this, &FLidarPointCloudEditorViewportClient::OnAssetViewerSettingsChanged);
+
+	// Set correct flags according to current profile settings
+	SetAdvancedShowFlagsForScene(UAssetViewerSettings::Get()->Profiles[GetMutableDefault<UEditorPerProjectUserSettings>()->AssetViewerProfileIndex].bPostProcessingEnabled);
+}
+
+FLidarPointCloudEditorViewportClient::~FLidarPointCloudEditorViewportClient()
+{
+	UAssetViewerSettings::Get()->OnAssetViewerSettingsChanged().RemoveAll(this);
+}
+
+void FLidarPointCloudEditorViewportClient::Tick(float DeltaSeconds)
+{
+	FEditorViewportClient::Tick(DeltaSeconds);
+
+	// Tick the preview scene world.
+	if (!GIntraFrameDebuggingGameThread)
+	{
+		PreviewScene->GetWorld()->Tick(LEVELTICK_All, DeltaSeconds);
+	}
+}
+
+void FLidarPointCloudEditorViewportClient::MouseMove(FViewport* InViewport, int32 x, int32 y)
+{
+	FEditorViewportClient::MouseMove(InViewport, x, y);
+}
+
+bool FLidarPointCloudEditorViewportClient::InputKey(FViewport* InViewport, int32 ControllerId, FKey Key, EInputEvent Event, float AmountDepressed, bool Gamepad)
+{
+	bool bHandled = FEditorViewportClient::InputKey(InViewport, ControllerId, Key, Event, AmountDepressed, false);
+
+	// Handle viewport screenshot.
+	bHandled |= InputTakeScreenshot(InViewport, Key, Event);
+
+	bHandled |= AdvancedPreviewScene->HandleInputKey(InViewport, ControllerId, Key, Event, AmountDepressed, Gamepad);
+
+	const bool bAltPressed = Viewport->KeyState(EKeys::LeftAlt) || Viewport->KeyState(EKeys::RightAlt);
+	const bool bCtrlPressed = Viewport->KeyState(EKeys::LeftControl) || Viewport->KeyState(EKeys::RightControl);
+	const bool bShiftPressed = Viewport->KeyState(EKeys::LeftShift) || Viewport->KeyState(EKeys::RightShift);
+
+	// Edit Mode
+	if (PointCloudEditorPtr.Pin()->IsEditMode())
+	{
+		if (Event == IE_Pressed)
+		{
+			if (Key == EKeys::Delete)
+			{
+				if (bShiftPressed)
+				{
+					PointCloudEditorPtr.Pin()->DeletePoints();
+				}
+				else
+				{
+					PointCloudEditorPtr.Pin()->HidePoints();
+				}
+
+				bHandled = true;
+			}
+			else if (Key == EKeys::Escape)
+			{
+				PointCloudEditorPtr.Pin()->DeselectPoints();
+				bHandled = true;
+			}
+		}
+
+		// Store current selection mode
+		ELidarPointCloudSelectionMode LastSelectionMode = SelectionMode;
+
+		// Update selection status
+		if (Key == EKeys::LeftMouseButton)
+		{
+			if (Viewport->KeyState(EKeys::LeftMouseButton))
+			{
+				if (bShiftPressed)
+				{
+					SelectionMode = ELidarPointCloudSelectionMode::Replace;
+				}
+				else if (bCtrlPressed)
+				{
+					SelectionMode = ELidarPointCloudSelectionMode::Add;
+				}
+				else if (bAltPressed)
+				{
+					SelectionMode = ELidarPointCloudSelectionMode::Subtract;
+				}
+				else
+				{
+					SelectionMode = ELidarPointCloudSelectionMode::None;
+				}
+			}
+			else
+			{
+				SelectionMode = ELidarPointCloudSelectionMode::None;
+			}
+
+			bHandled = true;
+		}
+
+		// Check if selection mode has changed since the last frame
+		if (SelectionMode != LastSelectionMode)
+		{
+			// Selection start
+			if (SelectionMode != ELidarPointCloudSelectionMode::None)
+			{
+				// Mark the cursor location for selection start
+				InViewport->GetMousePos(SelectionStartLocation);
+
+				OnSelectionStart();
+			}
+			// Selection end
+			else
+			{
+				FIntPoint SelectionEndLocation;
+
+				// Mark the cursor location for selection end
+				InViewport->GetMousePos(SelectionEndLocation);
+
+				FIntVector4 SelectionArea;
+
+				SelectionArea.X = FMath::Min(SelectionStartLocation.X, SelectionEndLocation.X);
+				SelectionArea.Y = FMath::Min(SelectionStartLocation.Y, SelectionEndLocation.Y);
+				SelectionArea.Z = FMath::Max(SelectionStartLocation.X, SelectionEndLocation.X);
+				SelectionArea.W = FMath::Max(SelectionStartLocation.Y, SelectionEndLocation.Y);
+
+				OnSelectionEnd(SelectionArea, LastSelectionMode);
+			}
+		}
+	}
+	// Navigation Mode
+	else
+	{
+
+	}
+
+	return bHandled;
+}
+
+bool FLidarPointCloudEditorViewportClient::InputAxis(FViewport* InViewport, int32 ControllerId, FKey Key, float Delta, float DeltaTime, int32 NumSamples, bool bGamepad)
+{
+	bool bResult = true;
+
+	if (!bDisableInput)
+	{
+		if (SelectionMode != ELidarPointCloudSelectionMode::None)
+		{
+			// Do not move the camera while in selection
+		}
+		else
+		{
+			bResult = AdvancedPreviewScene->HandleViewportInput(InViewport, ControllerId, Key, Delta, DeltaTime, NumSamples, bGamepad);
+			if (bResult)
+			{
+				Invalidate();
+			}
+			else
+			{
+				bResult = FEditorViewportClient::InputAxis(InViewport, ControllerId, Key, Delta, DeltaTime, NumSamples, bGamepad);
+			}
+		}
+	}
+
+	return bResult;
+}
+
+void FLidarPointCloudEditorViewportClient::ProcessClick(FSceneView& InView, HHitProxy* HitProxy, FKey Key, EInputEvent Event, uint32 HitX, uint32 HitY) { }
+
+void FLidarPointCloudEditorViewportClient::Draw(const FSceneView* View, FPrimitiveDrawInterface* PDI)
+{
+	FEditorViewportClient::Draw(View, PDI);
+}
+
+void FLidarPointCloudEditorViewportClient::DrawCanvas(FViewport& InViewport, FSceneView& View, FCanvas& Canvas)
+{
+	auto PointCloudEditor = PointCloudEditorPtr.Pin();
+	auto PointCloudEditorViewport = PointCloudEditorViewportPtr.Pin();
+	if (!PointCloudEditor.IsValid() || !PointCloudEditorViewport.IsValid())
+	{
+		return;
+	}
+
+	TArray<SLidarPointCloudEditorViewport::FOverlayTextItem> TextItems;
+
+	if (PointCloudEditorPtr.Pin()->IsEditMode())
+	{
+		// Text entries
+		{
+			TextItems.Add(SLidarPointCloudEditorViewport::FOverlayTextItem(FText::FromString("[ESCAPE] to de-select")));
+			TextItems.Add(SLidarPointCloudEditorViewport::FOverlayTextItem(FText::FromString("[SHIFT] + Mouse Drag to set selection")));
+			TextItems.Add(SLidarPointCloudEditorViewport::FOverlayTextItem(FText::FromString("[CTRL] + Mouse Drag to add selection")));
+			TextItems.Add(SLidarPointCloudEditorViewport::FOverlayTextItem(FText::FromString("[ALT] + Mouse Drag to subtract selection")));
+			TextItems.Add(SLidarPointCloudEditorViewport::FOverlayTextItem(FText::FromString("")));
+			TextItems.Add(SLidarPointCloudEditorViewport::FOverlayTextItem(FText::FromString("Press [DELETE] to hide selected")));
+			TextItems.Add(SLidarPointCloudEditorViewport::FOverlayTextItem(FText::FromString("Press [SHIFT] + [DELETE] to delete points")));
+			TextItems.Add(SLidarPointCloudEditorViewport::FOverlayTextItem(FText::FromString("")));
+			TextItems.Add(SLidarPointCloudEditorViewport::FOverlayTextItem(FText::Format(FText::FromString("Selected Points:  {0}"), FText::AsNumber(PointCloudEditorPtr.Pin()->GetSelectedPoints().Num()))));
+		}
+
+		// Selection box
+		if (SelectionMode != ELidarPointCloudSelectionMode::None)
+		{
+			FIntPoint SelectionCurrentLocation;
+			Viewport->GetMousePos(SelectionCurrentLocation);
+
+			float X = FMath::Min(SelectionStartLocation.X, SelectionCurrentLocation.X);
+			float Y = FMath::Min(SelectionStartLocation.Y, SelectionCurrentLocation.Y);
+			float SizeX = FMath::Max(SelectionStartLocation.X, SelectionCurrentLocation.X) - X;
+			float SizeY = FMath::Max(SelectionStartLocation.Y, SelectionCurrentLocation.Y) - Y;
+
+			FLinearColor SelectionColor = GetDefault<UEditorStyleSettings>()->SelectionColor;
+			SelectionColor.A = 0.35f;
+
+			Canvas.DrawTile(X, Y, SizeX, SizeY, 0, 0, 0, 0, SelectionColor);
+		}
+	}
+
+	PointCloudEditorViewport->PopulateOverlayText(TextItems);
+}
+
+bool FLidarPointCloudEditorViewportClient::ShouldOrbitCamera() const
+{
+	if (GetDefault<ULevelEditorViewportSettings>()->bUseUE3OrbitControls)
+	{
+		// this editor orbits always if ue3 orbit controls are enabled
+		return true;
+	}
+
+	return FEditorViewportClient::ShouldOrbitCamera();
+}
+
+void FLidarPointCloudEditorViewportClient::PerspectiveCameraMoved()
+{
+	FEditorViewportClient::PerspectiveCameraMoved();
+
+	// If in the process of transitioning to a new location, don't update the orbit camera position.
+	// On the final update of the transition, we will get here with IsPlaying()==false, and the editor camera position will
+	// be correctly updated.
+	if (GetViewTransform().IsPlaying())
+	{
+		return;
+	}
+
+	ToggleOrbitCamera(bUsingOrbitCamera);
+}
+
+void FLidarPointCloudEditorViewportClient::OnAssetViewerSettingsChanged(const FName& InPropertyName)
+{
+	if (InPropertyName == GET_MEMBER_NAME_CHECKED(FPreviewSceneProfile, bPostProcessingEnabled) || InPropertyName == NAME_None)
+	{
+		UAssetViewerSettings* Settings = UAssetViewerSettings::Get();
+		const int32 ProfileIndex = AdvancedPreviewScene->GetCurrentProfileIndex();
+		if (Settings->Profiles.IsValidIndex(ProfileIndex))
+		{
+			SetAdvancedShowFlagsForScene(Settings->Profiles[ProfileIndex].bPostProcessingEnabled);
+		}
+	}
+}
+
+void FLidarPointCloudEditorViewportClient::SetAdvancedShowFlagsForScene(const bool bAdvancedShowFlags)
+{
+	if (bAdvancedShowFlags)
+	{
+		EngineShowFlags.EnableAdvancedFeatures();
+	}
+	else
+	{
+		EngineShowFlags.DisableAdvancedFeatures();
+	}
+}
+
+void FLidarPointCloudEditorViewportClient::OnSelectionStart()
+{
+}
+
+void FLidarPointCloudEditorViewportClient::OnSelectionEnd(const FIntVector4& SelectionArea, ELidarPointCloudSelectionMode Mode)
+{
+	if (PointCloudComponent)
+	{
+		// Compute a view.
+		FSceneViewFamilyContext ViewFamily(FSceneViewFamily::ConstructionValues(Viewport, GetScene(), EngineShowFlags).SetRealtimeUpdate(IsRealtime()));
+		FSceneView* View = CalcSceneView(&ViewFamily);
+
+		/**
+		 * Compute 4 rays from the corners of the selection area
+		 *
+		 * 0 -- 3
+		 * |    |
+		 * 1 -- 2
+		 */
+		FVector Origins[6];
+		FVector Normals[6];
+		FVector Directions[4];
+		FSceneView::DeprojectScreenToWorld(FVector2D(SelectionArea.X, SelectionArea.Y), FIntRect(FIntPoint(0, 0), Viewport->GetSizeXY()), View->ViewMatrices.GetInvViewProjectionMatrix(), Origins[0], Directions[0]);
+		FSceneView::DeprojectScreenToWorld(FVector2D(SelectionArea.X, SelectionArea.W), FIntRect(FIntPoint(0, 0), Viewport->GetSizeXY()), View->ViewMatrices.GetInvViewProjectionMatrix(), Origins[1], Directions[1]);
+		FSceneView::DeprojectScreenToWorld(FVector2D(SelectionArea.Z, SelectionArea.W), FIntRect(FIntPoint(0, 0), Viewport->GetSizeXY()), View->ViewMatrices.GetInvViewProjectionMatrix(), Origins[2], Directions[2]);
+		FSceneView::DeprojectScreenToWorld(FVector2D(SelectionArea.Z, SelectionArea.Y), FIntRect(FIntPoint(0, 0), Viewport->GetSizeXY()), View->ViewMatrices.GetInvViewProjectionMatrix(), Origins[3], Directions[3]);
+
+		const FVector& ViewDirection = View->GetViewDirection();
+
+		// Shared calculations
+		Normals[4] = ViewDirection;
+		Normals[5] = -ViewDirection;
+		Origins[4] = Origins[0] + ViewDirection * 99999999;
+
+		// Perspective View
+		if (View->IsPerspectiveProjection())
+		{			
+			Origins[5] = Origins[0];
+
+			// Calculate plane normals
+			Normals[0] = Directions[1] ^ Directions[0];
+			Normals[1] = Directions[2] ^ Directions[1];
+			Normals[2] = Directions[3] ^ Directions[2];
+			Normals[3] = Directions[0] ^ Directions[3];
+
+			Normals[0].Normalize();
+			Normals[1].Normalize();
+			Normals[2].Normalize();
+			Normals[3].Normalize();
+		}
+		// Ortho Views
+		else
+		{
+			Origins[5] = -Origins[4];
+
+			// Left / Right view
+			if (ViewDirection.X != 0)
+			{
+				Normals[0] = -ViewDirection.X * FVector::RightVector;
+				Normals[1] = -ViewDirection.X * FVector::UpVector;
+				Normals[2] = ViewDirection.X * FVector::RightVector;
+				Normals[3] = ViewDirection.X * FVector::UpVector;
+			}
+			// Front / Back view
+			else if (ViewDirection.Y != 0)
+			{
+				Normals[0] = ViewDirection.Y * FVector::ForwardVector;
+				Normals[1] = ViewDirection.Y * FVector::UpVector;
+				Normals[2] = -ViewDirection.Y * FVector::ForwardVector;
+				Normals[3] = -ViewDirection.Y * FVector::UpVector;
+			}
+			// Top / Bottom view
+			else
+			{
+				Normals[0] = ViewDirection.Z * FVector::ForwardVector;
+				Normals[1] = -ViewDirection.Z * FVector::RightVector;
+				Normals[2] = -ViewDirection.Z * FVector::ForwardVector;
+				Normals[3] = ViewDirection.Z * FVector::RightVector;
+			}
+		}
+
+		// Compute selection frustum planes from the rays
+		TArray<FPlane, TInlineAllocator<6>> Planes;
+
+		// Left, Bottom, Right, Top, Forward, Backward
+		for(int32 i = 0; i < 6; ++i)
+		{
+			Planes.Emplace(Origins[i], Normals[i]);
+		}
+
+		if (Mode == ELidarPointCloudSelectionMode::Subtract)
+		{
+			PointCloudEditorPtr.Pin()->DeselectPointsByFrustum(FConvexVolume(Planes));
+		}
+		else
+		{
+			PointCloudEditorPtr.Pin()->SelectPointsByFrustum(FConvexVolume(Planes), Mode == ELidarPointCloudSelectionMode::Add);
+		}
+	}
+}
+
+void FLidarPointCloudEditorViewportClient::ResetCamera()
+{
+	if (PointCloudComponent != nullptr)
+	{
+		FocusViewportOnBox(PointCloudComponent->Bounds.GetBox());
+		Invalidate();
+	}
+}
+
+void FLidarPointCloudEditorViewportClient::ToggleShowNodes()
+{
+	if (PointCloudComponent != nullptr)
+	{
+		PointCloudComponent->bDrawNodeBounds = !PointCloudComponent->bDrawNodeBounds;
+		Invalidate();
+	}
+}
+
+bool FLidarPointCloudEditorViewportClient::IsSetShowNodesChecked() const
+{
+	return PointCloudComponent && PointCloudComponent->bDrawNodeBounds;
+}
+
+void FLidarPointCloudEditorViewportClient::SetPreviewCloud(ULidarPointCloud* InPointCloud, ULidarPointCloudComponent* InPointCloudComponent, bool bResetCamera)
+{
+	PointCloud = InPointCloud;
+	PointCloudComponent = InPointCloudComponent;
+
+	if (PointCloudComponent != nullptr)
+	{
+		PointCloudComponent->bOwnedByEditor = true;
+		PointCloudComponent->MarkRenderStateDirty();
+	}
+
+	if (bResetCamera)
+	{
+		ResetCamera();
+	}
+}
+
+#undef LOCTEXT_NAMESPACE 
\ No newline at end of file
diff --git a/Engine/Plugins/Enterprise/LidarPointCloud/Source/LidarPointCloudEditor/Private/LidarPointCloudEditorViewportClient.h b/Engine/Plugins/Enterprise/LidarPointCloud/Source/LidarPointCloudEditor/Private/LidarPointCloudEditorViewportClient.h
new file mode 100644
--- /dev/null
+++ b/Engine/Plugins/Enterprise/LidarPointCloud/Source/LidarPointCloudEditor/Private/LidarPointCloudEditorViewportClient.h
@@ -0,0 +1,103 @@
+// Copyright Epic Games, Inc. All Rights Reserved.
+ 
+#pragma once
+
+#include "CoreMinimal.h"
+#include "InputCoreTypes.h"
+#include "UnrealWidget.h"
+#include "EditorViewportClient.h"
+#include "Components.h"
+
+class FAdvancedPreviewScene;
+class FCanvas;
+class FLidarPointCloudEditor;
+class SLidarPointCloudEditorViewport;
+class ULidarPointCloud;
+class ULidarPointCloudComponent;
+
+enum class ELidarPointCloudSelectionMode : uint8
+{
+	None,
+	Add,
+	Subtract,
+	Replace
+};
+
+/** Viewport Client for the preview viewport */
+class FLidarPointCloudEditorViewportClient : public FEditorViewportClient, public TSharedFromThis<FLidarPointCloudEditorViewportClient>
+{
+public:
+	FLidarPointCloudEditorViewportClient(TWeakPtr<FLidarPointCloudEditor> InPointCloudEditor, const TSharedRef<SLidarPointCloudEditorViewport>& InPointCloudEditorViewport, const TSharedRef<FAdvancedPreviewScene>& InPreviewScene, ULidarPointCloud* InPreviewPointCLoud, ULidarPointCloudComponent* InPreviewPointCloudComponent);
+	~FLidarPointCloudEditorViewportClient();
+
+	// FEditorViewportClient interface
+	virtual void MouseMove(FViewport* Viewport, int32 x, int32 y) override;
+	virtual bool InputKey(FViewport* Viewport, int32 ControllerId, FKey Key, EInputEvent Event, float AmountDepressed = 1.f, bool bGamepad = false) override;
+	virtual bool InputAxis(FViewport* Viewport, int32 ControllerId, FKey Key, float Delta, float DeltaTime, int32 NumSamples = 1, bool bGamepad = false) override;
+	virtual void ProcessClick(class FSceneView& View, class HHitProxy* HitProxy, FKey Key, EInputEvent Event, uint32 HitX, uint32 HitY) override;
+	virtual void Tick(float DeltaSeconds) override;
+	virtual void Draw(const FSceneView* View, FPrimitiveDrawInterface* PDI) override;
+	virtual void DrawCanvas(FViewport& InViewport, FSceneView& View, FCanvas& Canvas) override;
+	virtual bool InputWidgetDelta(FViewport* InViewport, EAxisList::Type CurrentAxis, FVector& Drag, FRotator& Rot, FVector& Scale) override { return false; }
+	virtual void TrackingStarted(const struct FInputEventState& InInputState, bool bIsDragging, bool bNudge) override {}
+	virtual void TrackingStopped() override {}
+	virtual FWidget::EWidgetMode GetWidgetMode() const override { return FWidget::WM_None; }
+	virtual void SetWidgetMode(FWidget::EWidgetMode NewMode) override {}
+	virtual bool CanSetWidgetMode(FWidget::EWidgetMode NewMode) const override { return false; }
+	virtual bool CanCycleWidgetMode() const override { return false; }
+	virtual FVector GetWidgetLocation() const override { return FVector::ZeroVector; }
+	virtual FMatrix GetWidgetCoordSystem() const override { return FMatrix::Identity; }
+	virtual ECoordSystem GetWidgetCoordSystemSpace() const override { return COORD_Local; }
+	virtual bool ShouldOrbitCamera() const override;
+
+	/**
+	 *	Updates the point cloud and point cloud component being used in the Point Cloud Editor.
+	 *
+	 *	@param	InPointCloud				The point cloud handled by the editor.
+	 *	@param	InPointCloudComponent		The point cloud component from the viewport's scene.
+	 */
+	void SetPreviewCloud(ULidarPointCloud* InPointCloud, ULidarPointCloudComponent* InPointCloudComponent, bool bResetCamera = true);
+
+	void ResetCamera();
+
+	/** Callback for toggling the nodes show flag. */
+	void ToggleShowNodes();
+
+	/** Callback for checking the nodes show flag. */
+	bool IsSetShowNodesChecked() const;
+
+protected:
+	// FEditorViewportClient interface
+	virtual void PerspectiveCameraMoved() override;
+
+	/** Call back for when the user changes preview scene settings in the UI */
+	void OnAssetViewerSettingsChanged(const FName& InPropertyName);
+	/** Used to (re)-set the viewport show flags related to post processing*/
+	void SetAdvancedShowFlagsForScene(const bool bAdvancedShowFlags);
+
+private:
+	void OnSelectionStart();
+	void OnSelectionEnd(const FIntVector4& SelectionArea, ELidarPointCloudSelectionMode Mode);
+
+private:
+	/** Component for the point cloud. */
+	ULidarPointCloudComponent* PointCloudComponent;
+
+	/** The point cloud being used in the editor. */
+	ULidarPointCloud* PointCloud;
+
+	/** Pointer back to the PointCloud editor tool that owns us */
+	TWeakPtr<FLidarPointCloudEditor> PointCloudEditorPtr;
+
+	/** Pointer back to the PointCloudEditor viewport control that owns us */
+	TWeakPtr<SLidarPointCloudEditorViewport> PointCloudEditorViewportPtr;
+
+	/** Stored pointer to the preview scene in which the point cloud is shown */
+	FAdvancedPreviewScene* AdvancedPreviewScene;
+
+private:
+	ELidarPointCloudSelectionMode SelectionMode;
+
+	/** Location of the mouse cursor when the selection has started */
+	FIntPoint SelectionStartLocation;
+};
diff --git a/Engine/Plugins/Enterprise/LidarPointCloud/Source/LidarPointCloudEditor/Private/LidarPointCloudFactory.cpp b/Engine/Plugins/Enterprise/LidarPointCloud/Source/LidarPointCloudEditor/Private/LidarPointCloudFactory.cpp
new file mode 100644
--- /dev/null
+++ b/Engine/Plugins/Enterprise/LidarPointCloud/Source/LidarPointCloudEditor/Private/LidarPointCloudFactory.cpp
@@ -0,0 +1,265 @@
+// Copyright Epic Games, Inc. All Rights Reserved.
+
+#include "LidarPointCloudFactory.h"
+#include "LidarPointCloudShared.h"
+#include "LidarPointCloudSettings.h"
+#include "LidarPointCloudImportUI.h"
+#include "LidarPointCloudEditor.h"
+#include "IO/LidarPointCloudFileIO.h"
+#include "EditorStyleSet.h"
+#include "Framework/MultiBox/MultiBoxBuilder.h"
+#include "Editor.h"
+#include "Misc/ScopedSlowTask.h"
+
+#define LOCTEXT_NAMESPACE "LidarPointCloud"
+
+FText FAssetTypeActions_LidarPointCloud::GetName() const
+{
+	return NSLOCTEXT("AssetTypeActions", "FAssetTypeActions_LidarPointCloud", "LiDAR Point Cloud");
+}
+
+void FAssetTypeActions_LidarPointCloud::GetActions(const TArray<UObject*>& InObjects, FMenuBuilder& MenuBuilder)
+{
+	TArray<ULidarPointCloud*> PointClouds;
+	for (auto ObjIt = InObjects.CreateConstIterator(); ObjIt; ++ObjIt)
+	{
+		PointClouds.Add(CastChecked<ULidarPointCloud>(*ObjIt));
+	}
+
+	// Make sure at least one asset has source assigned
+	bool bSourceExists = false;
+	for (auto& PointCloud : PointClouds)
+	{
+		if (!PointCloud->GetSourcePath().IsEmpty())
+		{
+			bSourceExists = true;
+			break;
+		}
+	}
+
+	if (bSourceExists)
+	{
+		MenuBuilder.AddMenuEntry(
+			LOCTEXT("LidarPointCloud_Reimport", "Reimport Selected"),
+			LOCTEXT("LidarPointCloud_ReimportTooltip", "Reimports this point cloud asset."),
+			FSlateIcon(FEditorStyle::GetStyleSetName(), "TextureEditor.Reimport"),
+			FUIAction(
+				FExecuteAction::CreateSP(this, &FAssetTypeActions_LidarPointCloud::ExecuteReimport, PointClouds),
+				FCanExecuteAction()
+			)
+		);
+	}
+
+	if (PointClouds.Num() > 1)
+	{
+		MenuBuilder.AddMenuEntry(
+			LOCTEXT("LidarPointCloud_Merge", "Merge Selected"),
+			LOCTEXT("LidarPointCloud_MergeTooltip", "Merges selected point cloud assets."),
+			FSlateIcon("LidarPointCloudStyle", "LidarPointCloudEditor.Merge"),
+			FUIAction(
+				FExecuteAction::CreateSP(this, &FAssetTypeActions_LidarPointCloud::ExecuteMerge, PointClouds),
+				FCanExecuteAction()
+			)
+		);
+
+		MenuBuilder.AddMenuEntry(
+			LOCTEXT("LidarPointCloud_Align", "Align Selected"),
+			LOCTEXT("LidarPointCloud_AlignTooltip", "Aligns selected point cloud assets."),
+			FSlateIcon("LidarPointCloudStyle", "LidarPointCloudEditor.Align"),
+			FUIAction(
+				FExecuteAction::CreateSP(this, &FAssetTypeActions_LidarPointCloud::ExecuteAlign, PointClouds),
+				FCanExecuteAction()
+			)
+		);
+	}
+}
+
+void FAssetTypeActions_LidarPointCloud::OpenAssetEditor(const TArray<UObject*>& InObjects, TSharedPtr<class IToolkitHost> EditWithinLevelEditor /*= TSharedPtr<IToolkitHost>()*/)
+{
+	const EToolkitMode::Type Mode = EditWithinLevelEditor.IsValid() ? EToolkitMode::WorldCentric : EToolkitMode::Standalone;
+
+	for (auto ObjIt = InObjects.CreateConstIterator(); ObjIt; ++ObjIt)
+	{
+		if (ULidarPointCloud* PointCloud = Cast<ULidarPointCloud>(*ObjIt))
+		{
+			TSharedRef<FLidarPointCloudEditor> NewPointCloudEditor(new FLidarPointCloudEditor());
+			NewPointCloudEditor->InitPointCloudEditor(Mode, EditWithinLevelEditor, PointCloud);
+		}
+	}
+}
+
+void FAssetTypeActions_LidarPointCloud::ExecuteReimport(TArray<ULidarPointCloud*> PointClouds)
+{
+	for (ULidarPointCloud* PC : PointClouds)
+	{
+		PC->Reimport(GetDefault<ULidarPointCloudSettings>()->bUseAsyncImport);
+	}
+}
+
+void FAssetTypeActions_LidarPointCloud::ExecuteMerge(TArray<ULidarPointCloud*> PointClouds)
+{
+	if (PointClouds.Num() < 2)
+	{
+		return;
+	}
+
+	FString PackageName = PointClouds[0]->GetOutermost()->GetName() + TEXT("_Merged");
+	UPackage* MergedCloudPackage = CreatePackage(nullptr, *PackageName);
+	MergedCloudPackage->SetPackageFlags(PKG_NewlyCreated);
+
+	ULidarPointCloud* PC = NewObject<ULidarPointCloud>(MergedCloudPackage, FName(*FPaths::GetBaseFilename(PackageName)), EObjectFlags::RF_Public | EObjectFlags::RF_Standalone | EObjectFlags::RF_Transactional);
+	if (IsValid(PC))
+	{
+		TArray<FString> Names({ "Initializing", "Self" });
+
+		for (ULidarPointCloud* Asset : PointClouds)
+		{
+			Names.Add(FPaths::GetBaseFilename(FStringAssetReference(Asset).ToString()));
+		}
+
+		FScopedSlowTask ProgressDialog(PointClouds.Num() + 2, LOCTEXT("Merge", "Merging Point Clouds..."));
+		ProgressDialog.MakeDialog();
+		int32 i = 0;
+
+		PC->Merge(PointClouds, [&ProgressDialog, &i, &Names](float Progress) {
+			ProgressDialog.EnterProgressFrame(1.f, FText::FromString(Names[i++]));
+		});
+
+		PC->MarkPackageDirty();
+	}
+}
+
+void FAssetTypeActions_LidarPointCloud::ExecuteAlign(TArray<ULidarPointCloud*> PointClouds)
+{
+	FScopedSlowTask ProgressDialog(1, LOCTEXT("Align", "Aligning Point Clouds..."));
+	ProgressDialog.MakeDialog();
+	ProgressDialog.EnterProgressFrame(1.f);
+	ULidarPointCloud::AlignClouds(PointClouds);
+}
+
+ULidarPointCloudFactory::ULidarPointCloudFactory()
+{
+	bImportingAll = false;
+	bCreateNew = true;
+	bEditorImport = true;
+	SupportedClass = ULidarPointCloud::StaticClass();
+
+	TArray<FString> exts = ULidarPointCloudFileIO::GetSupportedImportExtensions();
+	for (FString ext : exts)
+	{
+		Formats.Add(*ext.Append(";LiDAR Point Cloud"));
+	}
+}
+
+UObject* ULidarPointCloudFactory::ImportObject(UClass* InClass, UObject* InOuter, FName InName, EObjectFlags Flags, const FString& Filename, const TCHAR* Parms, bool& OutCanceled)
+{
+	bCreateNew = false;
+	UObject* NewPC = Super::ImportObject(InClass, InOuter, InName, Flags, Filename, Parms, OutCanceled);
+	bCreateNew = true;
+
+	return NewPC;
+}
+
+UObject* ULidarPointCloudFactory::FactoryCreateFile(UClass* InClass, UObject* InParent, FName InName, EObjectFlags Flags, const FString& Filename, const TCHAR* Parms, FFeedbackContext* Warn, bool& bOutOperationCanceled)
+{
+	UObject *OutObject = nullptr;
+
+	UImportSubsystem* ImportSubsystem = GEditor->GetEditorSubsystem<UImportSubsystem>();
+	ImportSubsystem->OnAssetPreImport.Broadcast(this, InClass, InParent, InName, *FPaths::GetExtension(Filename));
+
+	// Check if the headers differ between files. Log occurrence and prompt the user for interaction if necessary
+	if (bImportingAll && ImportSettings.IsValid())
+	{
+		if (!ImportSettings->IsFileCompatible(Filename))
+		{
+			PC_WARNING("Inconsistent header information between files - batch import cancelled.");
+			bImportingAll = false;
+		}
+		else
+		{
+			ImportSettings->SetNewFilename(Filename);
+		}
+	}
+
+	if (!bImportingAll)
+	{
+		ImportSettings = FLidarPointCloudImportUI::ShowImportDialog(Filename, false);
+		bImportingAll = ImportSettings.IsValid() && ImportSettings->bImportAll;
+	}
+
+	if (ImportSettings.IsValid())
+	{
+		OutObject = ULidarPointCloud::CreateFromFile(Filename, ImportSettings->Clone(), InParent, InName, Flags);
+	}
+	else
+	{
+		bOutOperationCanceled = true;
+	}
+
+	ImportSubsystem->OnAssetPostImport.Broadcast(this, OutObject);
+
+	return OutObject;
+}
+
+UObject* ULidarPointCloudFactory::FactoryCreateNew(UClass* InClass, UObject* InParent, FName InName, EObjectFlags Flags, UObject* Context, FFeedbackContext* Warn)
+{
+	return NewObject<ULidarPointCloud>(InParent, InName, Flags);
+}
+
+bool ULidarPointCloudFactory::DoesSupportClass(UClass* Class) { return Class == ULidarPointCloud::StaticClass(); }
+
+bool ULidarPointCloudFactory::CanReimport(UObject* Obj, TArray<FString>& OutFilenames)
+{	
+	ULidarPointCloud* PC = Cast<ULidarPointCloud>(Obj);
+	if (PC)
+	{
+		OutFilenames.Add(*PC->GetSourcePath());
+		return true;
+	}
+	return false;
+}
+
+void ULidarPointCloudFactory::SetReimportPaths(UObject* Obj, const TArray<FString>& NewReimportPaths)
+{
+	ULidarPointCloud *PC = Cast<ULidarPointCloud>(Obj);
+	if (PC && NewReimportPaths.Num())
+	{
+		PC->SetSourcePath(NewReimportPaths[0]);
+	}
+}
+
+EReimportResult::Type ULidarPointCloudFactory::Reimport(UObject* Obj)
+{
+	ULidarPointCloud *PC = Cast<ULidarPointCloud>(Obj);
+
+	if (PC)
+	{	
+		bool bSuccess = false;
+
+		// Show existing settings, if the cloud has any
+		if (PC->ImportSettings.IsValid())
+		{
+			bSuccess = FLidarPointCloudImportUI::ShowImportDialog(PC->ImportSettings, true);
+		}
+		// ... otherwise attempt to generate new, based on the source path (if valid)
+		else if (FPaths::FileExists(PC->GetSourcePath()))
+		{
+			PC->ImportSettings = FLidarPointCloudImportUI::ShowImportDialog(PC->GetSourcePath(), true);
+			bSuccess = PC->ImportSettings.IsValid();
+		}
+		else
+		{
+			PC_ERROR("Cannot reimport, source path is incorrect.");
+		}
+
+		if (bSuccess)
+		{
+			PC->Reimport(GetDefault<ULidarPointCloudSettings>()->bUseAsyncImport);
+		}
+	}
+
+	// Return cancelled, to avoid showing 2 notifications
+	return EReimportResult::Cancelled;
+}
+
+#undef LOCTEXT_NAMESPACE
\ No newline at end of file
diff --git a/Engine/Plugins/Enterprise/LidarPointCloud/Source/LidarPointCloudEditor/Private/LidarPointCloudFactory.h b/Engine/Plugins/Enterprise/LidarPointCloud/Source/LidarPointCloudEditor/Private/LidarPointCloudFactory.h
new file mode 100644
--- /dev/null
+++ b/Engine/Plugins/Enterprise/LidarPointCloud/Source/LidarPointCloudEditor/Private/LidarPointCloudFactory.h
@@ -0,0 +1,58 @@
+// Copyright Epic Games, Inc. All Rights Reserved.
+
+#pragma once
+#include "EditorReimportHandler.h"
+#include "AssetTypeCategories.h"
+#include "AssetTypeActions_Base.h"
+#include "Factories/Factory.h"
+#include "LidarPointCloud.h"
+#include "LidarPointCloudFactory.generated.h"
+
+class FAssetTypeActions_LidarPointCloud : public FAssetTypeActions_Base
+{
+public:
+	FAssetTypeActions_LidarPointCloud() {}
+
+	// Begin IAssetTypeActions Interface
+	virtual FText GetName() const override;
+	virtual FColor GetTypeColor() const override { return FColor(0, 128, 128); }
+	virtual UClass* GetSupportedClass() const override { return ULidarPointCloud::StaticClass(); }
+	virtual bool HasActions(const TArray<UObject*>& InObjects) const override { return true; }
+	virtual void GetActions(const TArray<UObject*>& InObjects, FMenuBuilder& MenuBuilder) override;
+	virtual uint32 GetCategories() override { return EAssetTypeCategories::Misc; }
+	virtual void OpenAssetEditor(const TArray<UObject*>& InObjects, TSharedPtr<class IToolkitHost> EditWithinLevelEditor = TSharedPtr<IToolkitHost>()) override;
+	// End IAssetTypeActions Interface
+
+private:
+	void ExecuteReimport(TArray<ULidarPointCloud*> PointClouds);
+	void ExecuteMerge(TArray<ULidarPointCloud*> PointClouds);
+	void ExecuteAlign(TArray<ULidarPointCloud*> PointClouds);
+};
+
+UCLASS()
+class ULidarPointCloudFactory : public UFactory, public FReimportHandler
+{
+	GENERATED_BODY()
+
+private:
+	TSharedPtr<struct FLidarPointCloudImportSettings> ImportSettings;
+	bool bImportingAll;
+
+public:
+	ULidarPointCloudFactory();
+
+	// Begin UFactory Interface
+	virtual UObject* ImportObject(UClass* InClass, UObject* InOuter, FName InName, EObjectFlags Flags, const FString& Filename, const TCHAR* Parms, bool& OutCanceled) override;
+	virtual UObject* FactoryCreateFile(UClass* InClass, UObject* InParent, FName InName, EObjectFlags Flags, const FString& Filename, const TCHAR* Parms, FFeedbackContext* Warn, bool& bOutOperationCanceled) override;
+	virtual UObject* FactoryCreateNew(UClass* InClass, UObject* InParent, FName InName, EObjectFlags Flags, UObject* Context, FFeedbackContext* Warn) override;
+	virtual bool DoesSupportClass(UClass* Class) override;
+	virtual bool FactoryCanImport(const FString& Filename) override { return true; }
+	virtual FText GetDisplayName() const override { return FText::FromString("LiDAR Point Cloud"); }
+	// End UFactory Interface
+
+	// Begin FReimportHandler Interface
+	virtual bool CanReimport(UObject* Obj, TArray<FString>& OutFilenames) override;
+	virtual void SetReimportPaths(UObject* Obj, const TArray<FString>& NewReimportPaths) override;
+	virtual EReimportResult::Type Reimport(UObject* Obj) override;
+	// End FReimportHandler Interface
+};
\ No newline at end of file
diff --git a/Engine/Plugins/Enterprise/LidarPointCloud/Source/LidarPointCloudEditor/Private/LidarPointCloudStyle.cpp b/Engine/Plugins/Enterprise/LidarPointCloud/Source/LidarPointCloudEditor/Private/LidarPointCloudStyle.cpp
new file mode 100644
--- /dev/null
+++ b/Engine/Plugins/Enterprise/LidarPointCloud/Source/LidarPointCloudEditor/Private/LidarPointCloudStyle.cpp
@@ -0,0 +1,87 @@
+// Copyright Epic Games, Inc. All Rights Reserved.
+
+#include "LidarPointCloudStyle.h"
+#include "Interfaces/IPluginManager.h"
+#include "Styling/SlateStyleRegistry.h"
+#include "Styling/SlateTypes.h"
+#include "EditorStyleSet.h"
+#include "Interfaces/IPluginManager.h"
+
+#define IMAGE_PLUGIN_BRUSH( RelativePath, ... ) FSlateImageBrush( FLidarPointCloudStyle::InContent( RelativePath, ".png" ), __VA_ARGS__ )
+
+TSharedPtr<FSlateStyleSet> FLidarPointCloudStyle::StyleSet = nullptr;
+TSharedPtr<class ISlateStyle> FLidarPointCloudStyle::Get() { return StyleSet; }
+
+void FLidarPointCloudStyle::Initialize()
+{
+	// Const icon & thumbnail sizes
+	const FVector2D Icon16x16(16.0f, 16.0f);
+	const FVector2D Icon32x32(32.0f, 32.0f);
+	const FVector2D Icon40x40(40.0f, 40.0f);
+	const FVector2D Icon128x128(128.0f, 128.0f);
+
+	// Only register once
+	if (StyleSet.IsValid())
+	{
+		return;
+	}
+
+
+	StyleSet = MakeShareable(new FSlateStyleSet(GetStyleSetName()));
+	StyleSet->SetContentRoot(FPaths::EngineContentDir() / TEXT("Editor/Slate"));
+	StyleSet->SetCoreContentRoot(FPaths::EngineContentDir() / TEXT("Slate"));
+
+	FString ContentDir = IPluginManager::Get().FindPlugin(TEXT("LidarPointCloud"))->GetContentDir();
+	StyleSet->SetContentRoot(ContentDir);
+
+	StyleSet->Set("ClassIcon.LidarPointCloud", new IMAGE_PLUGIN_BRUSH("icon_32", Icon16x16));
+	StyleSet->Set("ClassIcon32.LidarPointCloud", new IMAGE_PLUGIN_BRUSH("icon_32", Icon32x32));
+	StyleSet->Set("ClassThumbnail.LidarPointCloud", new IMAGE_PLUGIN_BRUSH("icon_128", Icon128x128));
+
+	StyleSet->Set("ClassIcon.LidarPointCloudActor", new IMAGE_PLUGIN_BRUSH("icon_32", Icon16x16));
+	StyleSet->Set("ClassThumbnail.LidarPointCloudActor", new IMAGE_PLUGIN_BRUSH("icon_128", Icon128x128));
+
+	StyleSet->Set("ClassIcon.LidarPointCloudComponent", new IMAGE_PLUGIN_BRUSH("icon_32", Icon16x16));
+	StyleSet->Set("ClassThumbnail.LidarPointCloudComponent", new IMAGE_PLUGIN_BRUSH("icon_128", Icon128x128));
+
+	StyleSet->Set("LidarPointCloudEditor.BuildCollision", new IMAGE_PLUGIN_BRUSH("icon_collision_40", Icon40x40));
+	StyleSet->Set("LidarPointCloudEditor.RemoveCollision", new IMAGE_PLUGIN_BRUSH("icon_removecollision_40", Icon40x40));
+
+	StyleSet->Set("LidarPointCloudEditor.EditMode", new IMAGE_PLUGIN_BRUSH("icon_edit_40", Icon40x40));
+	StyleSet->Set("LidarPointCloudEditor.HideSelected", new IMAGE_PLUGIN_BRUSH("icon_hideselected_40", Icon40x40));
+	StyleSet->Set("LidarPointCloudEditor.UnhideAll", new IMAGE_PLUGIN_BRUSH("icon_unhideall_40", Icon40x40));
+
+	StyleSet->Set("LidarPointCloudEditor.DeleteSelected", new IMAGE_PLUGIN_BRUSH("icon_deleteselected_40", Icon40x40));
+
+	StyleSet->Set("LidarPointCloudEditor.Extract", new IMAGE_PLUGIN_BRUSH("icon_extract_40", Icon40x40));
+	StyleSet->Set("LidarPointCloudEditor.ExtractCopy", new IMAGE_PLUGIN_BRUSH("icon_extractcopy_40", Icon40x40));
+	StyleSet->Set("LidarPointCloudEditor.Merge", new IMAGE_PLUGIN_BRUSH("icon_merge_40", Icon40x40));
+	StyleSet->Set("LidarPointCloudEditor.Align", new IMAGE_PLUGIN_BRUSH("icon_align_40", Icon40x40));
+
+	FSlateStyleRegistry::RegisterSlateStyle(*StyleSet.Get());
+}
+
+#undef IMAGE_PLUGIN_BRUSH
+
+void FLidarPointCloudStyle::Shutdown()
+{
+	if (StyleSet.IsValid())
+	{
+		FSlateStyleRegistry::UnRegisterSlateStyle(*StyleSet.Get());
+		ensure(StyleSet.IsUnique());
+		StyleSet.Reset();
+	}
+}
+
+FName FLidarPointCloudStyle::GetStyleSetName()
+{
+	static FName PaperStyleName(TEXT("LidarPointCloudStyle"));
+	return PaperStyleName;
+}
+
+FString FLidarPointCloudStyle::InContent(const FString& RelativePath, const ANSICHAR* Extension)
+{
+	static FString ContentDir = IPluginManager::Get().FindPlugin(TEXT("LidarPointCloud"))->GetContentDir() / TEXT("Icons");
+	return (ContentDir / RelativePath) + Extension;
+}
+
