Copy up Dev-Enterprise @ CL11133089

diff --git a/Engine/Plugins/Enterprise/LidarPointCloud/Source/LidarPointCloudEditor/Private/LidarPointCloudStyle.h b/Engine/Plugins/Enterprise/LidarPointCloud/Source/LidarPointCloudEditor/Private/LidarPointCloudStyle.h
new file mode 100644
--- /dev/null
+++ b/Engine/Plugins/Enterprise/LidarPointCloud/Source/LidarPointCloudEditor/Private/LidarPointCloudStyle.h
@@ -0,0 +1,24 @@
+// Copyright Epic Games, Inc. All Rights Reserved.
+
+#pragma once
+
+#include "CoreMinimal.h"
+#include "Styling/ISlateStyle.h"
+#include "Styling/SlateStyle.h"
+
+class FLidarPointCloudStyle
+{
+public:
+	static void Initialize();
+
+	static void Shutdown();
+
+	static TSharedPtr< class ISlateStyle > Get();
+
+	static FName GetStyleSetName();
+private:
+	static FString InContent(const FString& RelativePath, const ANSICHAR* Extension);
+
+private:
+	static TSharedPtr< class FSlateStyleSet > StyleSet;
+};
diff --git a/Engine/Plugins/Enterprise/LidarPointCloud/Source/LidarPointCloudRuntime/Private/Collision/LidarPointCloudCollision.cpp b/Engine/Plugins/Enterprise/LidarPointCloud/Source/LidarPointCloudRuntime/Private/Collision/LidarPointCloudCollision.cpp
new file mode 100644
--- /dev/null
+++ b/Engine/Plugins/Enterprise/LidarPointCloud/Source/LidarPointCloudRuntime/Private/Collision/LidarPointCloudCollision.cpp
@@ -0,0 +1,543 @@
+// Copyright Epic Games, Inc. All Rights Reserved.
+ 
+#include "LidarPointCloudCollision.h"
+#include "LidarPointCloudShared.h"
+#include "LidarPointCloudOctree.h"
+
+#include "Interfaces/Interface_CollisionDataProvider.h"
+#include "Async/Async.h"
+#include "Misc/ScopeLock.h"
+
+// Stores triangle indices
+int8 TriangleTable[256][16] =
+{
+	{-1, -1, -1, -1, -1, -1, -1, -1, -1, -1, -1, -1, -1, -1, -1, -1},
+	{0, 8, 3, -1, -1, -1, -1, -1, -1, -1, -1, -1, -1, -1, -1, -1},
+	{0, 1, 9, -1, -1, -1, -1, -1, -1, -1, -1, -1, -1, -1, -1, -1},
+	{1, 8, 3, 9, 8, 1, -1, -1, -1, -1, -1, -1, -1, -1, -1, -1},
+	{1, 2, 10, -1, -1, -1, -1, -1, -1, -1, -1, -1, -1, -1, -1, -1},
+	{0, 8, 3, 1, 2, 10, -1, -1, -1, -1, -1, -1, -1, -1, -1, -1},
+	{9, 2, 10, 0, 2, 9, -1, -1, -1, -1, -1, -1, -1, -1, -1, -1},
+	{2, 8, 3, 2, 10, 8, 10, 9, 8, -1, -1, -1, -1, -1, -1, -1},
+	{3, 11, 2, -1, -1, -1, -1, -1, -1, -1, -1, -1, -1, -1, -1, -1},
+	{0, 11, 2, 8, 11, 0, -1, -1, -1, -1, -1, -1, -1, -1, -1, -1},
+	{1, 9, 0, 2, 3, 11, -1, -1, -1, -1, -1, -1, -1, -1, -1, -1},
+	{1, 11, 2, 1, 9, 11, 9, 8, 11, -1, -1, -1, -1, -1, -1, -1},
+	{3, 10, 1, 11, 10, 3, -1, -1, -1, -1, -1, -1, -1, -1, -1, -1},
+	{0, 10, 1, 0, 8, 10, 8, 11, 10, -1, -1, -1, -1, -1, -1, -1},
+	{3, 9, 0, 3, 11, 9, 11, 10, 9, -1, -1, -1, -1, -1, -1, -1},
+	{9, 8, 10, 10, 8, 11, -1, -1, -1, -1, -1, -1, -1, -1, -1, -1},
+	{4, 7, 8, -1, -1, -1, -1, -1, -1, -1, -1, -1, -1, -1, -1, -1},
+	{4, 3, 0, 7, 3, 4, -1, -1, -1, -1, -1, -1, -1, -1, -1, -1},
+	{0, 1, 9, 8, 4, 7, -1, -1, -1, -1, -1, -1, -1, -1, -1, -1},
+	{4, 1, 9, 4, 7, 1, 7, 3, 1, -1, -1, -1, -1, -1, -1, -1},
+	{1, 2, 10, 8, 4, 7, -1, -1, -1, -1, -1, -1, -1, -1, -1, -1},
+	{3, 4, 7, 3, 0, 4, 1, 2, 10, -1, -1, -1, -1, -1, -1, -1},
+	{9, 2, 10, 9, 0, 2, 8, 4, 7, -1, -1, -1, -1, -1, -1, -1},
+	{2, 10, 9, 2, 9, 7, 2, 7, 3, 7, 9, 4, -1, -1, -1, -1},
+	{8, 4, 7, 3, 11, 2, -1, -1, -1, -1, -1, -1, -1, -1, -1, -1},
+	{11, 4, 7, 11, 2, 4, 2, 0, 4, -1, -1, -1, -1, -1, -1, -1},
+	{9, 0, 1, 8, 4, 7, 2, 3, 11, -1, -1, -1, -1, -1, -1, -1},
+	{4, 7, 11, 9, 4, 11, 9, 11, 2, 9, 2, 1, -1, -1, -1, -1},
+	{3, 10, 1, 3, 11, 10, 7, 8, 4, -1, -1, -1, -1, -1, -1, -1},
+	{1, 11, 10, 1, 4, 11, 1, 0, 4, 7, 11, 4, -1, -1, -1, -1},
+	{4, 7, 8, 9, 0, 11, 9, 11, 10, 11, 0, 3, -1, -1, -1, -1},
+	{4, 7, 11, 4, 11, 9, 9, 11, 10, -1, -1, -1, -1, -1, -1, -1},
+	{9, 5, 4, -1, -1, -1, -1, -1, -1, -1, -1, -1, -1, -1, -1, -1},
+	{9, 5, 4, 0, 8, 3, -1, -1, -1, -1, -1, -1, -1, -1, -1, -1},
+	{0, 5, 4, 1, 5, 0, -1, -1, -1, -1, -1, -1, -1, -1, -1, -1},
+	{8, 5, 4, 8, 3, 5, 3, 1, 5, -1, -1, -1, -1, -1, -1, -1},
+	{1, 2, 10, 9, 5, 4, -1, -1, -1, -1, -1, -1, -1, -1, -1, -1},
+	{3, 0, 8, 1, 2, 10, 4, 9, 5, -1, -1, -1, -1, -1, -1, -1},
+	{5, 2, 10, 5, 4, 2, 4, 0, 2, -1, -1, -1, -1, -1, -1, -1},
+	{2, 10, 5, 3, 2, 5, 3, 5, 4, 3, 4, 8, -1, -1, -1, -1},
+	{9, 5, 4, 2, 3, 11, -1, -1, -1, -1, -1, -1, -1, -1, -1, -1},
+	{0, 11, 2, 0, 8, 11, 4, 9, 5, -1, -1, -1, -1, -1, -1, -1},
+	{0, 5, 4, 0, 1, 5, 2, 3, 11, -1, -1, -1, -1, -1, -1, -1},
+	{2, 1, 5, 2, 5, 8, 2, 8, 11, 4, 8, 5, -1, -1, -1, -1},
+	{10, 3, 11, 10, 1, 3, 9, 5, 4, -1, -1, -1, -1, -1, -1, -1},
+	{4, 9, 5, 0, 8, 1, 8, 10, 1, 8, 11, 10, -1, -1, -1, -1},
+	{5, 4, 0, 5, 0, 11, 5, 11, 10, 11, 0, 3, -1, -1, -1, -1},
+	{5, 4, 8, 5, 8, 10, 10, 8, 11, -1, -1, -1, -1, -1, -1, -1},
+	{9, 7, 8, 5, 7, 9, -1, -1, -1, -1, -1, -1, -1, -1, -1, -1},
+	{9, 3, 0, 9, 5, 3, 5, 7, 3, -1, -1, -1, -1, -1, -1, -1},
+	{0, 7, 8, 0, 1, 7, 1, 5, 7, -1, -1, -1, -1, -1, -1, -1},
+	{1, 5, 3, 3, 5, 7, -1, -1, -1, -1, -1, -1, -1, -1, -1, -1},
+	{9, 7, 8, 9, 5, 7, 10, 1, 2, -1, -1, -1, -1, -1, -1, -1},
+	{10, 1, 2, 9, 5, 0, 5, 3, 0, 5, 7, 3, -1, -1, -1, -1},
+	{8, 0, 2, 8, 2, 5, 8, 5, 7, 10, 5, 2, -1, -1, -1, -1},
+	{2, 10, 5, 2, 5, 3, 3, 5, 7, -1, -1, -1, -1, -1, -1, -1},
+	{7, 9, 5, 7, 8, 9, 3, 11, 2, -1, -1, -1, -1, -1, -1, -1},
+	{9, 5, 7, 9, 7, 2, 9, 2, 0, 2, 7, 11, -1, -1, -1, -1},
+	{2, 3, 11, 0, 1, 8, 1, 7, 8, 1, 5, 7, -1, -1, -1, -1},
+	{11, 2, 1, 11, 1, 7, 7, 1, 5, -1, -1, -1, -1, -1, -1, -1},
+	{9, 5, 8, 8, 5, 7, 10, 1, 3, 10, 3, 11, -1, -1, -1, -1},
+	{5, 7, 0, 5, 0, 9, 7, 11, 0, 1, 0, 10, 11, 10, 0, -1},
+	{11, 10, 0, 11, 0, 3, 10, 5, 0, 8, 0, 7, 5, 7, 0, -1},
+	{11, 10, 5, 7, 11, 5, -1, -1, -1, -1, -1, -1, -1, -1, -1, -1},
+	{10, 6, 5, -1, -1, -1, -1, -1, -1, -1, -1, -1, -1, -1, -1, -1},
+	{0, 8, 3, 5, 10, 6, -1, -1, -1, -1, -1, -1, -1, -1, -1, -1},
+	{9, 0, 1, 5, 10, 6, -1, -1, -1, -1, -1, -1, -1, -1, -1, -1},
+	{1, 8, 3, 1, 9, 8, 5, 10, 6, -1, -1, -1, -1, -1, -1, -1},
+	{1, 6, 5, 2, 6, 1, -1, -1, -1, -1, -1, -1, -1, -1, -1, -1},
+	{1, 6, 5, 1, 2, 6, 3, 0, 8, -1, -1, -1, -1, -1, -1, -1},
+	{9, 6, 5, 9, 0, 6, 0, 2, 6, -1, -1, -1, -1, -1, -1, -1},
+	{5, 9, 8, 5, 8, 2, 5, 2, 6, 3, 2, 8, -1, -1, -1, -1},
+	{2, 3, 11, 10, 6, 5, -1, -1, -1, -1, -1, -1, -1, -1, -1, -1},
+	{11, 0, 8, 11, 2, 0, 10, 6, 5, -1, -1, -1, -1, -1, -1, -1},
+	{0, 1, 9, 2, 3, 11, 5, 10, 6, -1, -1, -1, -1, -1, -1, -1},
+	{5, 10, 6, 1, 9, 2, 9, 11, 2, 9, 8, 11, -1, -1, -1, -1},
+	{6, 3, 11, 6, 5, 3, 5, 1, 3, -1, -1, -1, -1, -1, -1, -1},
+	{0, 8, 11, 0, 11, 5, 0, 5, 1, 5, 11, 6, -1, -1, -1, -1},
+	{3, 11, 6, 0, 3, 6, 0, 6, 5, 0, 5, 9, -1, -1, -1, -1},
+	{6, 5, 9, 6, 9, 11, 11, 9, 8, -1, -1, -1, -1, -1, -1, -1},
+	{5, 10, 6, 4, 7, 8, -1, -1, -1, -1, -1, -1, -1, -1, -1, -1},
+	{4, 3, 0, 4, 7, 3, 6, 5, 10, -1, -1, -1, -1, -1, -1, -1},
+	{1, 9, 0, 5, 10, 6, 8, 4, 7, -1, -1, -1, -1, -1, -1, -1},
+	{10, 6, 5, 1, 9, 7, 1, 7, 3, 7, 9, 4, -1, -1, -1, -1},
+	{6, 1, 2, 6, 5, 1, 4, 7, 8, -1, -1, -1, -1, -1, -1, -1},
+	{1, 2, 5, 5, 2, 6, 3, 0, 4, 3, 4, 7, -1, -1, -1, -1},
+	{8, 4, 7, 9, 0, 5, 0, 6, 5, 0, 2, 6, -1, -1, -1, -1},
+	{7, 3, 9, 7, 9, 4, 3, 2, 9, 5, 9, 6, 2, 6, 9, -1},
+	{3, 11, 2, 7, 8, 4, 10, 6, 5, -1, -1, -1, -1, -1, -1, -1},
+	{5, 10, 6, 4, 7, 2, 4, 2, 0, 2, 7, 11, -1, -1, -1, -1},
+	{0, 1, 9, 4, 7, 8, 2, 3, 11, 5, 10, 6, -1, -1, -1, -1},
+	{9, 2, 1, 9, 11, 2, 9, 4, 11, 7, 11, 4, 5, 10, 6, -1},
+	{8, 4, 7, 3, 11, 5, 3, 5, 1, 5, 11, 6, -1, -1, -1, -1},
+	{5, 1, 11, 5, 11, 6, 1, 0, 11, 7, 11, 4, 0, 4, 11, -1},
+	{0, 5, 9, 0, 6, 5, 0, 3, 6, 11, 6, 3, 8, 4, 7, -1},
+	{6, 5, 9, 6, 9, 11, 4, 7, 9, 7, 11, 9, -1, -1, -1, -1},
+	{10, 4, 9, 6, 4, 10, -1, -1, -1, -1, -1, -1, -1, -1, -1, -1},
+	{4, 10, 6, 4, 9, 10, 0, 8, 3, -1, -1, -1, -1, -1, -1, -1},
+	{10, 0, 1, 10, 6, 0, 6, 4, 0, -1, -1, -1, -1, -1, -1, -1},
+	{8, 3, 1, 8, 1, 6, 8, 6, 4, 6, 1, 10, -1, -1, -1, -1},
+	{1, 4, 9, 1, 2, 4, 2, 6, 4, -1, -1, -1, -1, -1, -1, -1},
+	{3, 0, 8, 1, 2, 9, 2, 4, 9, 2, 6, 4, -1, -1, -1, -1},
+	{0, 2, 4, 4, 2, 6, -1, -1, -1, -1, -1, -1, -1, -1, -1, -1},
+	{8, 3, 2, 8, 2, 4, 4, 2, 6, -1, -1, -1, -1, -1, -1, -1},
+	{10, 4, 9, 10, 6, 4, 11, 2, 3, -1, -1, -1, -1, -1, -1, -1},
+	{0, 8, 2, 2, 8, 11, 4, 9, 10, 4, 10, 6, -1, -1, -1, -1},
+	{3, 11, 2, 0, 1, 6, 0, 6, 4, 6, 1, 10, -1, -1, -1, -1},
+	{6, 4, 1, 6, 1, 10, 4, 8, 1, 2, 1, 11, 8, 11, 1, -1},
+	{9, 6, 4, 9, 3, 6, 9, 1, 3, 11, 6, 3, -1, -1, -1, -1},
+	{8, 11, 1, 8, 1, 0, 11, 6, 1, 9, 1, 4, 6, 4, 1, -1},
+	{3, 11, 6, 3, 6, 0, 0, 6, 4, -1, -1, -1, -1, -1, -1, -1},
+	{6, 4, 8, 11, 6, 8, -1, -1, -1, -1, -1, -1, -1, -1, -1, -1},
+	{7, 10, 6, 7, 8, 10, 8, 9, 10, -1, -1, -1, -1, -1, -1, -1},
+	{0, 7, 3, 0, 10, 7, 0, 9, 10, 6, 7, 10, -1, -1, -1, -1},
+	{10, 6, 7, 1, 10, 7, 1, 7, 8, 1, 8, 0, -1, -1, -1, -1},
+	{10, 6, 7, 10, 7, 1, 1, 7, 3, -1, -1, -1, -1, -1, -1, -1},
+	{1, 2, 6, 1, 6, 8, 1, 8, 9, 8, 6, 7, -1, -1, -1, -1},
+	{2, 6, 9, 2, 9, 1, 6, 7, 9, 0, 9, 3, 7, 3, 9, -1},
+	{7, 8, 0, 7, 0, 6, 6, 0, 2, -1, -1, -1, -1, -1, -1, -1},
+	{7, 3, 2, 6, 7, 2, -1, -1, -1, -1, -1, -1, -1, -1, -1, -1},
+	{2, 3, 11, 10, 6, 8, 10, 8, 9, 8, 6, 7, -1, -1, -1, -1},
+	{2, 0, 7, 2, 7, 11, 0, 9, 7, 6, 7, 10, 9, 10, 7, -1},
+	{1, 8, 0, 1, 7, 8, 1, 10, 7, 6, 7, 10, 2, 3, 11, -1},
+	{11, 2, 1, 11, 1, 7, 10, 6, 1, 6, 7, 1, -1, -1, -1, -1},
+	{8, 9, 6, 8, 6, 7, 9, 1, 6, 11, 6, 3, 1, 3, 6, -1},
+	{0, 9, 1, 11, 6, 7, -1, -1, -1, -1, -1, -1, -1, -1, -1, -1},
+	{7, 8, 0, 7, 0, 6, 3, 11, 0, 11, 6, 0, -1, -1, -1, -1},
+	{7, 11, 6, -1, -1, -1, -1, -1, -1, -1, -1, -1, -1, -1, -1, -1},
+	{7, 6, 11, -1, -1, -1, -1, -1, -1, -1, -1, -1, -1, -1, -1, -1},
+	{3, 0, 8, 11, 7, 6, -1, -1, -1, -1, -1, -1, -1, -1, -1, -1},
+	{0, 1, 9, 11, 7, 6, -1, -1, -1, -1, -1, -1, -1, -1, -1, -1},
+	{8, 1, 9, 8, 3, 1, 11, 7, 6, -1, -1, -1, -1, -1, -1, -1},
+	{10, 1, 2, 6, 11, 7, -1, -1, -1, -1, -1, -1, -1, -1, -1, -1},
+	{1, 2, 10, 3, 0, 8, 6, 11, 7, -1, -1, -1, -1, -1, -1, -1},
+	{2, 9, 0, 2, 10, 9, 6, 11, 7, -1, -1, -1, -1, -1, -1, -1},
+	{6, 11, 7, 2, 10, 3, 10, 8, 3, 10, 9, 8, -1, -1, -1, -1},
+	{7, 2, 3, 6, 2, 7, -1, -1, -1, -1, -1, -1, -1, -1, -1, -1},
+	{7, 0, 8, 7, 6, 0, 6, 2, 0, -1, -1, -1, -1, -1, -1, -1},
+	{2, 7, 6, 2, 3, 7, 0, 1, 9, -1, -1, -1, -1, -1, -1, -1},
+	{1, 6, 2, 1, 8, 6, 1, 9, 8, 8, 7, 6, -1, -1, -1, -1},
+	{10, 7, 6, 10, 1, 7, 1, 3, 7, -1, -1, -1, -1, -1, -1, -1},
+	{10, 7, 6, 1, 7, 10, 1, 8, 7, 1, 0, 8, -1, -1, -1, -1},
+	{0, 3, 7, 0, 7, 10, 0, 10, 9, 6, 10, 7, -1, -1, -1, -1},
+	{7, 6, 10, 7, 10, 8, 8, 10, 9, -1, -1, -1, -1, -1, -1, -1},
+	{6, 8, 4, 11, 8, 6, -1, -1, -1, -1, -1, -1, -1, -1, -1, -1},
+	{3, 6, 11, 3, 0, 6, 0, 4, 6, -1, -1, -1, -1, -1, -1, -1},
+	{8, 6, 11, 8, 4, 6, 9, 0, 1, -1, -1, -1, -1, -1, -1, -1},
+	{9, 4, 6, 9, 6, 3, 9, 3, 1, 11, 3, 6, -1, -1, -1, -1},
+	{6, 8, 4, 6, 11, 8, 2, 10, 1, -1, -1, -1, -1, -1, -1, -1},
+	{1, 2, 10, 3, 0, 11, 0, 6, 11, 0, 4, 6, -1, -1, -1, -1},
+	{4, 11, 8, 4, 6, 11, 0, 2, 9, 2, 10, 9, -1, -1, -1, -1},
+	{10, 9, 3, 10, 3, 2, 9, 4, 3, 11, 3, 6, 4, 6, 3, -1},
+	{8, 2, 3, 8, 4, 2, 4, 6, 2, -1, -1, -1, -1, -1, -1, -1},
+	{0, 4, 2, 4, 6, 2, -1, -1, -1, -1, -1, -1, -1, -1, -1, -1},
+	{1, 9, 0, 2, 3, 4, 2, 4, 6, 4, 3, 8, -1, -1, -1, -1},
+	{1, 9, 4, 1, 4, 2, 2, 4, 6, -1, -1, -1, -1, -1, -1, -1},
+	{8, 1, 3, 8, 6, 1, 8, 4, 6, 6, 10, 1, -1, -1, -1, -1},
+	{10, 1, 0, 10, 0, 6, 6, 0, 4, -1, -1, -1, -1, -1, -1, -1},
+	{4, 6, 3, 4, 3, 8, 6, 10, 3, 0, 3, 9, 10, 9, 3, -1},
+	{10, 9, 4, 6, 10, 4, -1, -1, -1, -1, -1, -1, -1, -1, -1, -1},
+	{4, 9, 5, 7, 6, 11, -1, -1, -1, -1, -1, -1, -1, -1, -1, -1},
+	{0, 8, 3, 4, 9, 5, 11, 7, 6, -1, -1, -1, -1, -1, -1, -1},
+	{5, 0, 1, 5, 4, 0, 7, 6, 11, -1, -1, -1, -1, -1, -1, -1},
+	{11, 7, 6, 8, 3, 4, 3, 5, 4, 3, 1, 5, -1, -1, -1, -1},
+	{9, 5, 4, 10, 1, 2, 7, 6, 11, -1, -1, -1, -1, -1, -1, -1},
+	{6, 11, 7, 1, 2, 10, 0, 8, 3, 4, 9, 5, -1, -1, -1, -1},
+	{7, 6, 11, 5, 4, 10, 4, 2, 10, 4, 0, 2, -1, -1, -1, -1},
+	{3, 4, 8, 3, 5, 4, 3, 2, 5, 10, 5, 2, 11, 7, 6, -1},
+	{7, 2, 3, 7, 6, 2, 5, 4, 9, -1, -1, -1, -1, -1, -1, -1},
+	{9, 5, 4, 0, 8, 6, 0, 6, 2, 6, 8, 7, -1, -1, -1, -1},
+	{3, 6, 2, 3, 7, 6, 1, 5, 0, 5, 4, 0, -1, -1, -1, -1},
+	{6, 2, 8, 6, 8, 7, 2, 1, 8, 4, 8, 5, 1, 5, 8, -1},
+	{9, 5, 4, 10, 1, 6, 1, 7, 6, 1, 3, 7, -1, -1, -1, -1},
+	{1, 6, 10, 1, 7, 6, 1, 0, 7, 8, 7, 0, 9, 5, 4, -1},
+	{4, 0, 10, 4, 10, 5, 0, 3, 10, 6, 10, 7, 3, 7, 10, -1},
+	{7, 6, 10, 7, 10, 8, 5, 4, 10, 4, 8, 10, -1, -1, -1, -1},
+	{6, 9, 5, 6, 11, 9, 11, 8, 9, -1, -1, -1, -1, -1, -1, -1},
+	{3, 6, 11, 0, 6, 3, 0, 5, 6, 0, 9, 5, -1, -1, -1, -1},
+	{0, 11, 8, 0, 5, 11, 0, 1, 5, 5, 6, 11, -1, -1, -1, -1},
+	{6, 11, 3, 6, 3, 5, 5, 3, 1, -1, -1, -1, -1, -1, -1, -1},
+	{1, 2, 10, 9, 5, 11, 9, 11, 8, 11, 5, 6, -1, -1, -1, -1},
+	{0, 11, 3, 0, 6, 11, 0, 9, 6, 5, 6, 9, 1, 2, 10, -1},
+	{11, 8, 5, 11, 5, 6, 8, 0, 5, 10, 5, 2, 0, 2, 5, -1},
+	{6, 11, 3, 6, 3, 5, 2, 10, 3, 10, 5, 3, -1, -1, -1, -1},
+	{5, 8, 9, 5, 2, 8, 5, 6, 2, 3, 8, 2, -1, -1, -1, -1},
+	{9, 5, 6, 9, 6, 0, 0, 6, 2, -1, -1, -1, -1, -1, -1, -1},
+	{1, 5, 8, 1, 8, 0, 5, 6, 8, 3, 8, 2, 6, 2, 8, -1},
+	{1, 5, 6, 2, 1, 6, -1, -1, -1, -1, -1, -1, -1, -1, -1, -1},
+	{1, 3, 6, 1, 6, 10, 3, 8, 6, 5, 6, 9, 8, 9, 6, -1},
+	{10, 1, 0, 10, 0, 6, 9, 5, 0, 5, 6, 0, -1, -1, -1, -1},
+	{0, 3, 8, 5, 6, 10, -1, -1, -1, -1, -1, -1, -1, -1, -1, -1},
+	{10, 5, 6, -1, -1, -1, -1, -1, -1, -1, -1, -1, -1, -1, -1, -1},
+	{11, 5, 10, 7, 5, 11, -1, -1, -1, -1, -1, -1, -1, -1, -1, -1},
+	{11, 5, 10, 11, 7, 5, 8, 3, 0, -1, -1, -1, -1, -1, -1, -1},
+	{5, 11, 7, 5, 10, 11, 1, 9, 0, -1, -1, -1, -1, -1, -1, -1},
+	{10, 7, 5, 10, 11, 7, 9, 8, 1, 8, 3, 1, -1, -1, -1, -1},
+	{11, 1, 2, 11, 7, 1, 7, 5, 1, -1, -1, -1, -1, -1, -1, -1},
+	{0, 8, 3, 1, 2, 7, 1, 7, 5, 7, 2, 11, -1, -1, -1, -1},
+	{9, 7, 5, 9, 2, 7, 9, 0, 2, 2, 11, 7, -1, -1, -1, -1},
+	{7, 5, 2, 7, 2, 11, 5, 9, 2, 3, 2, 8, 9, 8, 2, -1},
+	{2, 5, 10, 2, 3, 5, 3, 7, 5, -1, -1, -1, -1, -1, -1, -1},
+	{8, 2, 0, 8, 5, 2, 8, 7, 5, 10, 2, 5, -1, -1, -1, -1},
+	{9, 0, 1, 5, 10, 3, 5, 3, 7, 3, 10, 2, -1, -1, -1, -1},
+	{9, 8, 2, 9, 2, 1, 8, 7, 2, 10, 2, 5, 7, 5, 2, -1},
+	{1, 3, 5, 3, 7, 5, -1, -1, -1, -1, -1, -1, -1, -1, -1, -1},
+	{0, 8, 7, 0, 7, 1, 1, 7, 5, -1, -1, -1, -1, -1, -1, -1},
+	{9, 0, 3, 9, 3, 5, 5, 3, 7, -1, -1, -1, -1, -1, -1, -1},
+	{9, 8, 7, 5, 9, 7, -1, -1, -1, -1, -1, -1, -1, -1, -1, -1},
+	{5, 8, 4, 5, 10, 8, 10, 11, 8, -1, -1, -1, -1, -1, -1, -1},
+	{5, 0, 4, 5, 11, 0, 5, 10, 11, 11, 3, 0, -1, -1, -1, -1},
+	{0, 1, 9, 8, 4, 10, 8, 10, 11, 10, 4, 5, -1, -1, -1, -1},
+	{10, 11, 4, 10, 4, 5, 11, 3, 4, 9, 4, 1, 3, 1, 4, -1},
+	{2, 5, 1, 2, 8, 5, 2, 11, 8, 4, 5, 8, -1, -1, -1, -1},
+	{0, 4, 11, 0, 11, 3, 4, 5, 11, 2, 11, 1, 5, 1, 11, -1},
+	{0, 2, 5, 0, 5, 9, 2, 11, 5, 4, 5, 8, 11, 8, 5, -1},
+	{9, 4, 5, 2, 11, 3, -1, -1, -1, -1, -1, -1, -1, -1, -1, -1},
+	{2, 5, 10, 3, 5, 2, 3, 4, 5, 3, 8, 4, -1, -1, -1, -1},
+	{5, 10, 2, 5, 2, 4, 4, 2, 0, -1, -1, -1, -1, -1, -1, -1},
+	{3, 10, 2, 3, 5, 10, 3, 8, 5, 4, 5, 8, 0, 1, 9, -1},
+	{5, 10, 2, 5, 2, 4, 1, 9, 2, 9, 4, 2, -1, -1, -1, -1},
+	{8, 4, 5, 8, 5, 3, 3, 5, 1, -1, -1, -1, -1, -1, -1, -1},
+	{0, 4, 5, 1, 0, 5, -1, -1, -1, -1, -1, -1, -1, -1, -1, -1},
+	{8, 4, 5, 8, 5, 3, 9, 0, 5, 0, 3, 5, -1, -1, -1, -1},
+	{9, 4, 5, -1, -1, -1, -1, -1, -1, -1, -1, -1, -1, -1, -1, -1},
+	{4, 11, 7, 4, 9, 11, 9, 10, 11, -1, -1, -1, -1, -1, -1, -1},
+	{0, 8, 3, 4, 9, 7, 9, 11, 7, 9, 10, 11, -1, -1, -1, -1},
+	{1, 10, 11, 1, 11, 4, 1, 4, 0, 7, 4, 11, -1, -1, -1, -1},
+	{3, 1, 4, 3, 4, 8, 1, 10, 4, 7, 4, 11, 10, 11, 4, -1},
+	{4, 11, 7, 9, 11, 4, 9, 2, 11, 9, 1, 2, -1, -1, -1, -1},
+	{9, 7, 4, 9, 11, 7, 9, 1, 11, 2, 11, 1, 0, 8, 3, -1},
+	{11, 7, 4, 11, 4, 2, 2, 4, 0, -1, -1, -1, -1, -1, -1, -1},
+	{11, 7, 4, 11, 4, 2, 8, 3, 4, 3, 2, 4, -1, -1, -1, -1},
+	{2, 9, 10, 2, 7, 9, 2, 3, 7, 7, 4, 9, -1, -1, -1, -1},
+	{9, 10, 7, 9, 7, 4, 10, 2, 7, 8, 7, 0, 2, 0, 7, -1},
+	{3, 7, 10, 3, 10, 2, 7, 4, 10, 1, 10, 0, 4, 0, 10, -1},
+	{1, 10, 2, 8, 7, 4, -1, -1, -1, -1, -1, -1, -1, -1, -1, -1},
+	{4, 9, 1, 4, 1, 7, 7, 1, 3, -1, -1, -1, -1, -1, -1, -1},
+	{4, 9, 1, 4, 1, 7, 0, 8, 1, 8, 7, 1, -1, -1, -1, -1},
+	{4, 0, 3, 7, 4, 3, -1, -1, -1, -1, -1, -1, -1, -1, -1, -1},
+	{4, 8, 7, -1, -1, -1, -1, -1, -1, -1, -1, -1, -1, -1, -1, -1},
+	{9, 10, 8, 10, 11, 8, -1, -1, -1, -1, -1, -1, -1, -1, -1, -1},
+	{3, 0, 9, 3, 9, 11, 11, 9, 10, -1, -1, -1, -1, -1, -1, -1},
+	{0, 1, 10, 0, 10, 8, 8, 10, 11, -1, -1, -1, -1, -1, -1, -1},
+	{3, 1, 10, 11, 3, 10, -1, -1, -1, -1, -1, -1, -1, -1, -1, -1},
+	{1, 2, 11, 1, 11, 9, 9, 11, 8, -1, -1, -1, -1, -1, -1, -1},
+	{3, 0, 9, 3, 9, 11, 1, 2, 9, 2, 11, 9, -1, -1, -1, -1},
+	{0, 2, 11, 8, 0, 11, -1, -1, -1, -1, -1, -1, -1, -1, -1, -1},
+	{3, 2, 11, -1, -1, -1, -1, -1, -1, -1, -1, -1, -1, -1, -1, -1},
+	{2, 3, 8, 2, 8, 10, 10, 8, 9, -1, -1, -1, -1, -1, -1, -1},
+	{9, 10, 2, 0, 9, 2, -1, -1, -1, -1, -1, -1, -1, -1, -1, -1},
+	{2, 3, 8, 2, 8, 10, 0, 1, 8, 1, 10, 8, -1, -1, -1, -1},
+	{1, 10, 2, -1, -1, -1, -1, -1, -1, -1, -1, -1, -1, -1, -1, -1},
+	{1, 3, 8, 9, 1, 8, -1, -1, -1, -1, -1, -1, -1, -1, -1, -1},
+	{0, 9, 1, -1, -1, -1, -1, -1, -1, -1, -1, -1, -1, -1, -1, -1},
+	{0, 3, 8, -1, -1, -1, -1, -1, -1, -1, -1, -1, -1, -1, -1, -1},
+	{-1, -1, -1, -1, -1, -1, -1, -1, -1, -1, -1, -1, -1, -1, -1, -1}
+};
+
+/** Stores pre-computed data, ready for processing. */
+uint8 VertexData[256][181];
+
+FORCEINLINE FVector ExtractCoordinatesFromIndex(const int32& Index, const FIntVector& Dimensions)
+{
+	const int32 W = Index % (Dimensions.X * Dimensions.Y);
+	return FVector(W % Dimensions.X, W / Dimensions.X, Index / (Dimensions.X * Dimensions.Y));
+}
+
+void BuildScaledVertexTable(float CellSize)
+{
+	const float HalfCellSize = CellSize * 0.5f;
+
+	FVector VertexOffsetTable[12] =
+	{
+		FVector(HalfCellSize, 0, 0),
+		FVector(CellSize, HalfCellSize, 0),
+		FVector(HalfCellSize, CellSize, 0),
+		FVector(0, HalfCellSize, 0),
+		FVector(HalfCellSize, 0, CellSize),
+		FVector(CellSize, HalfCellSize, CellSize),
+		FVector(HalfCellSize, CellSize, CellSize),
+		FVector(0, HalfCellSize, CellSize),
+		FVector(0, 0, HalfCellSize),
+		FVector(CellSize, 0, HalfCellSize),
+		FVector(CellSize, CellSize, HalfCellSize),
+		FVector(0, CellSize, HalfCellSize)
+	};
+
+	int8* VertexList = nullptr;
+
+	TArray<FVector> TempArray;
+	for (int32 i = 0; i < 256; ++i)
+	{
+		VertexList = TriangleTable[i];
+		TempArray.Reset();
+
+		for (int32 t = 0; t < 16; ++t)
+		{
+			if (VertexList[t] == -1)
+			{
+				break;
+			}
+
+			TempArray.Add(VertexOffsetTable[VertexList[t]]);
+		}
+
+		VertexData[i][0] = TempArray.Num();
+		FMemory::Memcpy(VertexData[i] + 1, TempArray.GetData(), TempArray.Num() * 12);
+	}
+}
+
+/** Applies Marching Cubes algorith to extract mesh information from the given voxelized grid */
+void MarchingCubes(uint8* VoxelizedGrid, int32 GridDimension, float CellSize, FVector VertexOffset, TArray<FVector>& OutVertices)
+{
+	FVector Coords;
+	uint8* Vertices = nullptr;
+
+	const int32 GridDimensionSq = GridDimension * GridDimension;
+
+	uint8* Grid = VoxelizedGrid;
+	uint8* GridUp = Grid + GridDimensionSq;
+
+	for (int32 sZ = 0; sZ < GridDimension - 1; ++sZ)
+	{
+		Coords.Z = VertexOffset.Z + (sZ + 0.5f) * CellSize;
+
+		Grid = VoxelizedGrid + sZ * GridDimensionSq;
+		GridUp = Grid + GridDimensionSq;
+		int32 BaseIndex = 0;
+
+		for (int32 sY = 0; sY < GridDimension - 1; ++sY)
+		{
+			Coords.Y = VertexOffset.Y + (sY + 0.5f) * CellSize;
+
+			for (int32 sX = 0; sX < GridDimension - 1; ++sX, ++BaseIndex)
+			{
+				Coords.X = VertexOffset.X + (sX + 0.5f) * CellSize;
+
+				// Calculate the variation of the cube
+				uint8 CubeVariation = 0;
+				if (Grid[BaseIndex] == 0) { CubeVariation |= 1; }
+				if (Grid[BaseIndex + 1] == 0) { CubeVariation |= 2; }
+				if (Grid[BaseIndex + GridDimension + 1] == 0) { CubeVariation |= 4; }
+				if (Grid[BaseIndex + GridDimension] == 0) { CubeVariation |= 8; }
+				if (GridUp[BaseIndex] == 0) { CubeVariation |= 16; }
+				if (GridUp[BaseIndex + 1] == 0) { CubeVariation |= 32; }
+				if (GridUp[BaseIndex + GridDimension + 1] == 0) { CubeVariation |= 64; }
+				if (GridUp[BaseIndex + GridDimension] == 0) { CubeVariation |= 128; }
+
+				// No triangles to add
+				if (CubeVariation == 0 || CubeVariation == 255)
+				{
+					continue;
+				}
+
+				Vertices = VertexData[CubeVariation];
+				const uint8 NumVertices = Vertices[0];
+
+				if (NumVertices == 0)
+				{
+					continue;
+				}
+
+				int32 Offset = OutVertices.AddUninitialized(NumVertices);
+				for (FVector* Vertex = OutVertices.GetData() + Offset, *SourceData = (FVector*)(Vertices + 1), *DataEnd = Vertex + NumVertices; Vertex != DataEnd; ++Vertex, ++SourceData)
+				{
+					*Vertex = Coords + *SourceData;
+				}
+			}
+
+			++BaseIndex;
+		}
+	}
+}
+
+void LidarPointCloudCollision::BuildCollisionMesh(const FLidarPointCloudOctree* Octree, const float& CellSize, const bool& bVisibleOnly, FTriMeshCollisionData* CollisionMesh)
+{
+	if (!CollisionMesh)
+	{
+		return;
+	}
+
+	FBenchmarkTimer::Reset();
+
+	// Expand original bounds to make sure the mesh is closed at the edges
+	const FBox Bounds = Octree->GetBounds().ExpandBy(FVector::OneVector * CellSize, FVector::OneVector * CellSize);
+	const FVector BoundsSize = Bounds.GetSize();
+
+	// Number of cells in each axis
+	const int32 BatchSize = GetDefault<ULidarPointCloudSettings>()->MeshingBatchSize;
+
+	// Determines the number of samples to perform
+	const FIntVector NumSamples(BoundsSize / ((BatchSize - 1) * CellSize) + 1);
+	const int32 TotalNumSamples = NumSamples.X * NumSamples.Y * NumSamples.Z;
+
+	// Precalculated for performance reasons
+	const float InversedCellSize = 1 / CellSize;
+	const FBox BaseSamplingBounds(Bounds.Min, Bounds.Min + BatchSize * CellSize);
+
+	// Multithreading
+	const int32 MaxNumThreads = FMath::Min(FPlatformMisc::NumberOfCoresIncludingHyperthreads() - 1, TotalNumSamples);
+	TArray<TFuture<void>> ThreadResults;
+	FThreadSafeCounter SampleIndex = 0;
+
+	// Data storage
+	TArray<TArray<FVector>> Vertices;
+	Vertices.AddDefaulted(MaxNumThreads);
+	FThreadSafeCounter NumVertices = 0;
+
+	BuildScaledVertexTable(CellSize);
+	
+	// Fire threads
+	for (int32 t = 0; t < MaxNumThreads && SampleIndex.GetValue() < TotalNumSamples; t++)
+	{
+		ThreadResults.Add(Async(EAsyncExecution::TaskGraph, [t, &SampleIndex, &Vertices, &NumVertices, CollisionMesh, BatchSize, TotalNumSamples, NumSamples, BaseSamplingBounds, CellSize, Octree, InversedCellSize, bVisibleOnly]
+		{
+			// Local caching arrays to reduce number of syncs required
+			TArray<FVector> _Vertices;
+
+			const int32 NumCellsSq = BatchSize * BatchSize;
+			const int32 NumCellsCu = NumCellsSq * BatchSize;
+
+			uint8* VoxelizedGrid = new uint8[NumCellsCu];
+
+			const float OffsetMultiplier = (BatchSize - 1) * CellSize;
+
+			TArray<FLidarPointCloudPoint*> Selection;
+			
+			int32 Index;
+			while ((Index = SampleIndex.Add(1)) < TotalNumSamples)
+			{
+				const FBox SamplingBounds = BaseSamplingBounds.ShiftBy(ExtractCoordinatesFromIndex(Index, NumSamples) * OffsetMultiplier);
+
+				// Sample the data using the calculated bounds
+				Octree->GetPointsInBox(Selection, SamplingBounds, bVisibleOnly);
+
+				// Skip if no points in selection
+				if (Selection.Num() == 0)
+				{
+					continue;
+				}
+
+				// Clear existing grids
+				FMemory::Memzero(VoxelizedGrid, NumCellsCu);
+
+				// Calculate voxelized grids
+				for (auto Point = Selection.GetData(), DataEnd = Selection.GetData() + Selection.Num(); Point != DataEnd; ++Point)
+				{
+					// Calculate location relative to sampling bounds
+					FVector Location = (*Point)->Location - SamplingBounds.Min;
+
+					// Calculate grid coordinates
+					const FIntVector Grid(FMath::Min(BatchSize - 1, (int32)(Location.X * InversedCellSize)), FMath::Min(BatchSize - 1, (int32)(Location.Y * InversedCellSize)), FMath::Min(BatchSize - 1, (int32)(Location.Z * InversedCellSize)));
+
+					VoxelizedGrid[Grid.Z * NumCellsSq + Grid.Y * BatchSize + Grid.X] = 1;
+				}
+
+				MarchingCubes(VoxelizedGrid, BatchSize, CellSize, SamplingBounds.Min, _Vertices);
+			}
+
+			// Sync
+			Vertices[t].Append(_Vertices);
+			NumVertices.Add(_Vertices.Num());
+
+			// Clean up
+			delete[] VoxelizedGrid;
+		}));
+	}
+
+	// Sync threads
+	for (auto& ThreadResult : ThreadResults) ThreadResult.Get();
+
+	FBenchmarkTimer::Log("Collision: Meshing");
+
+	bool bWeldVertices = false;
+	if (bWeldVertices)
+	{
+		TSet<FVector> TmpVertexData;
+		TmpVertexData.Reserve(NumVertices.GetValue());
+		CollisionMesh->Vertices.Empty(NumVertices.GetValue());
+		CollisionMesh->Indices.Empty(NumVertices.GetValue() / 3);
+		CollisionMesh->Indices.AddUninitialized(NumVertices.GetValue() / 3);
+
+		bool bInSet = false;
+		int32* Indices = (int32*)CollisionMesh->Indices.GetData();
+
+		for (int32 i = 0, Offset = 0; i < Vertices.Num(); ++i)
+		{
+			for (FVector* Vertex = Vertices[i].GetData(), *DataEnd = Vertex + Vertices[i].Num(); Vertex != DataEnd; ++Vertex, ++Indices)
+			{
+				*Indices = TmpVertexData.Add(*Vertex, &bInSet).AsInteger();
+				if (!bInSet)
+				{
+					CollisionMesh->Vertices.Add(*Vertex);
+				}
+			}
+		}
+
+		CollisionMesh->Vertices.Shrink();
+
+		PC_LOG("Before: %d, After: %d / %d", NumVertices.GetValue(), CollisionMesh->Vertices.Num());
+	}
+	else
+	{
+		CollisionMesh->Vertices.Empty(NumVertices.GetValue());
+		CollisionMesh->Vertices.AddUninitialized(NumVertices.GetValue());
+
+		FVector* Dest = CollisionMesh->Vertices.GetData();
+
+		for (int32 i = 0, Offset = 0; i < Vertices.Num(); ++i)
+		{
+			const int32 NumElements = Vertices[i].Num();
+			if (NumElements > 0)
+			{
+				FMemory::Memcpy(Dest, Vertices[i].GetData(), NumElements * 12);
+				Dest += NumElements;
+			}
+		}
+
+		// Populate indices
+		CollisionMesh->Indices.Empty(NumVertices.GetValue() / 3);
+		CollisionMesh->Indices.AddUninitialized(NumVertices.GetValue() / 3);
+		int32 i = 0;
+		for (auto Index = (int32*)CollisionMesh->Indices.GetData(), DataEnd = Index + CollisionMesh->Indices.Num() * 3; Index != DataEnd; ++Index)
+		{
+			*Index = i++;
+		}
+	}	
+	
+	FBenchmarkTimer::Log("Collision: Exporting Data");
+}
diff --git a/Engine/Plugins/Enterprise/LidarPointCloud/Source/LidarPointCloudRuntime/Private/Collision/LidarPointCloudCollision.h b/Engine/Plugins/Enterprise/LidarPointCloud/Source/LidarPointCloudRuntime/Private/Collision/LidarPointCloudCollision.h
new file mode 100644
--- /dev/null
+++ b/Engine/Plugins/Enterprise/LidarPointCloud/Source/LidarPointCloudRuntime/Private/Collision/LidarPointCloudCollision.h
@@ -0,0 +1,13 @@
+// Copyright Epic Games, Inc. All Rights Reserved.
+
+#pragma once
+
+#include "CoreMinimal.h"
+
+class FLidarPointCloudOctree;
+struct FTriMeshCollisionData;
+
+namespace LidarPointCloudCollision
+{
+	void BuildCollisionMesh(const FLidarPointCloudOctree* Octree, const float& CellSize, const bool& bVisibleOnly, FTriMeshCollisionData* CollisionMesh);
+}
\ No newline at end of file
diff --git a/Engine/Plugins/Enterprise/LidarPointCloud/Source/LidarPointCloudRuntime/Private/IO/LidarPointCloudFileIO_LAS.cpp b/Engine/Plugins/Enterprise/LidarPointCloud/Source/LidarPointCloudRuntime/Private/IO/LidarPointCloudFileIO_LAS.cpp
new file mode 100644
--- /dev/null
+++ b/Engine/Plugins/Enterprise/LidarPointCloud/Source/LidarPointCloudRuntime/Private/IO/LidarPointCloudFileIO_LAS.cpp
@@ -0,0 +1,506 @@
+// Copyright Epic Games, Inc. All Rights Reserved.
+
+#include "IO/LidarPointCloudFileIO_LAS.h"
+#include "LidarPointCloudShared.h"
+#include "LidarPointCloud.h"
+#include "HAL/FileManager.h"
+#include "Misc/FileHelper.h"
+#include "Async/Async.h"
+#include "HAL/ThreadSafeCounter64.h"
+#include "Templates/Atomic.h"
+
+FLidarPointCloudImportSettings_LAS::FLidarPointCloudImportSettings_LAS(const FString& Filename)
+	: FLidarPointCloudImportSettings(Filename)
+{
+	ReadFileHeader(Filename);
+}
+
+void FLidarPointCloudImportSettings_LAS::Serialize(FArchive& Ar)
+{
+	FLidarPointCloudImportSettings::Serialize(Ar);
+
+	int32 Dummy32;
+	uint8 Dummy;
+	
+	if (Ar.CustomVer(ULidarPointCloud::PointCloudFileGUID) >= 11)
+	{
+	}
+	else if (Ar.CustomVer(ULidarPointCloud::PointCloudFileGUID) >= 10)
+	{
+		Ar << Dummy;
+	}
+	else if (Ar.CustomVer(ULidarPointCloud::PointCloudFileGUID) >= 7)
+	{
+		Ar << Dummy32 << Dummy32 << Dummy;
+	}
+}
+
+void FLidarPointCloudImportSettings_LAS::ReadFileHeader(const FString& InFilename)
+{
+	TUniquePtr<FArchive> Reader(IFileManager::Get().CreateFileReader(*Filename));
+	if (Reader)
+	{
+		int64 TotalSize = Reader->TotalSize();
+
+		// Check the minimum size for the file to be valid
+		if (TotalSize > 227)
+		{
+			// Reserve space for the full public header block
+			TArray<uint8> Data;
+			Data.AddUninitialized(375);
+
+			// Start by reading the basic set, for LAS file versions prior to 1.4
+			Reader->Serialize(Data.GetData(), 227);
+
+			// Populate the basic set
+			FMemory::Memcpy(&PublicHeaderBlock, Data.GetData(), 227);
+
+			uint16 BytesRead = 227;
+
+			// Check if the file has valid marker
+			if (PublicHeaderBlock.IsValid())
+			{
+				// Read extra 8 bytes added post 1.2
+				if (PublicHeaderBlock.VersionMinor > 2)
+				{
+					Reader->Serialize(Data.GetData() + BytesRead, 8);
+					BytesRead += 8;
+
+					// Read extra 140 bytes added post 1.3
+					if (PublicHeaderBlock.VersionMinor > 3)
+					{
+						Reader->Serialize(Data.GetData() + BytesRead, 140);
+						BytesRead += 140;
+					}
+
+					// Re-populate the header struct
+					FMemory::Memcpy(&PublicHeaderBlock, Data.GetData(), BytesRead);
+				}
+
+			}
+		}
+
+		Reader->Close();
+	}
+}
+
+/** Used to help track multiple import buffer allocations */
+class FLidarPointCloudFileIO_LAS_ImportBuffer
+{
+public:
+	FLidarPointCloudFileIO_LAS_ImportBuffer()
+		: bInUse(false)
+	{
+	}
+	~FLidarPointCloudFileIO_LAS_ImportBuffer() = default;
+	FLidarPointCloudFileIO_LAS_ImportBuffer(const FLidarPointCloudFileIO_LAS_ImportBuffer& Other)
+	{
+		Data = Other.Data;
+		bInUse = false;
+	}
+	FLidarPointCloudFileIO_LAS_ImportBuffer(FLidarPointCloudFileIO_LAS_ImportBuffer&&) = delete;
+	FLidarPointCloudFileIO_LAS_ImportBuffer& operator=(const FLidarPointCloudFileIO_LAS_ImportBuffer& Other)
+	{
+		Data = Other.Data;
+		bInUse = false;
+		return *this;
+	}
+	FLidarPointCloudFileIO_LAS_ImportBuffer& operator=(FLidarPointCloudFileIO_LAS_ImportBuffer&&) = delete;
+
+	FORCEINLINE uint8* GetData() { return Data.GetData(); }
+
+	void Initialize(const int32& Size)
+	{
+		Data.AddUninitialized(Size);
+	}
+
+public:
+	TAtomic<bool> bInUse;
+
+private:
+	TArray<uint8> Data;
+};
+
+/** Used to help track multiple import buffer allocations */
+class FLidarPointCloudFileIO_LAS_ImportBufferManager
+{
+public:
+	FLidarPointCloudFileIO_LAS_ImportBufferManager(const int32& BufferSize)
+		: BufferSize(BufferSize)
+		, Head(FLidarPointCloudFileIO_LAS_ImportBuffer())
+		, Tail(&Head)
+	{
+		Head.Element.Initialize(BufferSize);
+	}
+	~FLidarPointCloudFileIO_LAS_ImportBufferManager()
+	{
+		auto Iterator = &Head;
+		while (Iterator)
+		{
+			if (Iterator != &Head)
+			{
+				auto Tmp = Iterator;
+				Iterator = Iterator->Next;
+				delete Tmp;
+			}
+			else
+			{
+				Iterator = Iterator->Next;
+			}
+		}
+	}
+	FLidarPointCloudFileIO_LAS_ImportBufferManager(const FLidarPointCloudFileIO_LAS_ImportBufferManager&) = delete;
+	FLidarPointCloudFileIO_LAS_ImportBufferManager(FLidarPointCloudFileIO_LAS_ImportBufferManager&&) = delete;
+	FLidarPointCloudFileIO_LAS_ImportBufferManager& operator=(const FLidarPointCloudFileIO_LAS_ImportBufferManager&) = delete;
+	FLidarPointCloudFileIO_LAS_ImportBufferManager& operator=(FLidarPointCloudFileIO_LAS_ImportBufferManager&&) = delete;
+
+	FLidarPointCloudFileIO_LAS_ImportBuffer* GetFreeBuffer()
+	{
+		FLidarPointCloudFileIO_LAS_ImportBuffer* OutBuffer = nullptr;
+
+		// Find available memory allocation
+		{
+			auto Iterator = &Head;
+			while (Iterator)
+			{
+				if (!Iterator->Element.bInUse)
+				{
+					OutBuffer = &Iterator->Element;
+					break;
+				}
+
+				Iterator = Iterator->Next;
+			}
+		}
+
+		// If none found, add a new one
+		if (!OutBuffer)
+		{
+			Tail->Next = new TList<FLidarPointCloudFileIO_LAS_ImportBuffer>(FLidarPointCloudFileIO_LAS_ImportBuffer());
+			Tail = Tail->Next;
+			OutBuffer = &Tail->Element;
+			OutBuffer->Initialize(BufferSize);
+		}
+
+		OutBuffer->bInUse = true;
+
+		return OutBuffer;
+	}
+
+private:
+	int32 BufferSize;
+	TList<FLidarPointCloudFileIO_LAS_ImportBuffer> Head;
+	TList<FLidarPointCloudFileIO_LAS_ImportBuffer>* Tail;
+};
+
+bool ULidarPointCloudFileIO_LAS::HandleImport(const FString& Filename, TSharedPtr<FLidarPointCloudImportSettings> ImportSettings, FLidarPointCloudImportResults &OutImportResults)
+{
+	if (!ValidateImportSettings(ImportSettings, Filename))
+	{
+		return false;
+	}
+
+	FLidarPointCloudImportSettings_LAS* Settings = (FLidarPointCloudImportSettings_LAS*)ImportSettings.Get();
+	FLidarPointCloudFileIO_LAS_PublicHeaderBlock Header = Settings->PublicHeaderBlock;
+
+	// Return immediately if header is invalid
+	if (!Header.IsValid())
+	{
+		return false;
+	}
+
+	bool bSuccess = false;
+
+	TUniquePtr<FArchive> Reader(IFileManager::Get().CreateFileReader(*Filename));
+	if (Reader)
+	{
+		int64 TotalSize = Reader->TotalSize();
+		int64 TotalPointsToRead = Header.GetNumberOfPoints();
+
+		if (TotalPointsToRead > 0)
+		{
+			OutImportResults.SetMaxProgressCounter(TotalPointsToRead);
+			
+			bool bHasIntensityData = false;
+			bool bUse16BitIntensity = true;
+			bool bUse16BitRGB = true;
+
+			const int64 MaxBufferSize = GetDefault<ULidarPointCloudSettings>()->MaxImportBufferSize;
+
+			// Calculate max buffer size
+			const int64 MaxPointsToRead = FMath::Min(TotalPointsToRead, MaxBufferSize / Header.PointDataRecordLength);
+
+			// Detect bit depth
+			if(GetDefault<ULidarPointCloudSettings>()->bEnable8BitLASDetection)
+			{
+				bUse16BitIntensity = false;
+				bUse16BitRGB = false;
+
+				// Set the correct position for the reader
+				Reader->Seek(Header.OffsetToPointData);
+
+				// Calculate the amount of data to read
+				int64 PointsToRead = FMath::Min((int64)GetDefault<ULidarPointCloudSettings>()->MaxNumberOfPointsToScanLAS, TotalPointsToRead);
+
+				// Create data buffer
+				TArray<uint8> Data;
+				Data.AddUninitialized(PointsToRead * Header.PointDataRecordLength);
+
+				// Read the data
+				Reader->Serialize(Data.GetData(), Data.Num());
+
+				for(uint8* DataPtr = Data.GetData(), *DataEnd = DataPtr + Data.Num(); DataPtr != DataEnd; DataPtr += Header.PointDataRecordLength)
+				{
+					FLidarPointCloudFileIO_LAS_PointDataRecordFormatCommon* Record = (FLidarPointCloudFileIO_LAS_PointDataRecordFormatCommon*)DataPtr;
+
+					bHasIntensityData = bHasIntensityData || Record->Intensity > 0;
+					bUse16BitIntensity = bUse16BitIntensity || Record->Intensity > 255;
+
+					if(!bUse16BitRGB)
+					{
+						FLidarPointCloudFileIO_LAS_PointDataRecordFormatCommonRGB* RecordRGB = nullptr;
+
+						switch (Header.PointDataRecordFormat)
+						{
+						case 2: RecordRGB = (FLidarPointCloudFileIO_LAS_PointDataRecordFormat2*)DataPtr; break;
+						case 3: RecordRGB = (FLidarPointCloudFileIO_LAS_PointDataRecordFormat3*)DataPtr; break;
+						case 5: RecordRGB = (FLidarPointCloudFileIO_LAS_PointDataRecordFormat5*)DataPtr; break;
+						case 7: RecordRGB = (FLidarPointCloudFileIO_LAS_PointDataRecordFormat7*)DataPtr; break;
+						case 8: RecordRGB = (FLidarPointCloudFileIO_LAS_PointDataRecordFormat8*)DataPtr; break;
+						case 10: RecordRGB = (FLidarPointCloudFileIO_LAS_PointDataRecordFormat10*)DataPtr; break;
+						}
+
+						bUse16BitRGB = RecordRGB && (RecordRGB->Red > 255 || RecordRGB->Green > 255 || RecordRGB->Blue > 255);
+					}
+				}
+			}
+
+			// Read Data
+			{
+				// Set the correct position for the reader
+				Reader->Seek(Header.OffsetToPointData);
+
+				int64 PointsRead = 0;
+
+				// Clear any existing data
+				OutImportResults.Points.Empty(TotalPointsToRead);
+				OutImportResults.ClassificationsImported.Empty();
+
+				const float IntensityMultiplier = 1.0f / (bUse16BitIntensity ? 65535.0f : 255.0f);
+				const float RGBMultiplier = 1.0f / (bUse16BitRGB ? 65535.0f : 255.0f);
+
+				const bool bAutoCenter = GetDefault<ULidarPointCloudSettings>()->bAutoCenterOnImport;
+				const float ImportScale = GetDefault<ULidarPointCloudSettings>()->ImportScale;
+
+				FThreadSafeBool bFirstPointSet = !bAutoCenter;
+
+				// Multi-threading
+				FCriticalSection CoordsLock;
+				FCriticalSection PointsLock;
+				TArray<TFuture<void>> ThreadResults;
+				FLidarPointCloudFileIO_LAS_ImportBufferManager BufferManager(MaxPointsToRead * Header.PointDataRecordLength);
+
+				// Stream the data
+				while (PointsRead < TotalPointsToRead && !OutImportResults.IsCancelled())
+				{
+					FLidarPointCloudFileIO_LAS_ImportBuffer* Buffer = BufferManager.GetFreeBuffer();
+
+					// Data should never be null
+					check(Buffer->GetData());
+
+					// Calculate the amount of data to read
+					int64 PointsToRead = FMath::Min(MaxPointsToRead, TotalPointsToRead - PointsRead);
+
+					// Read the data
+					Reader->Serialize(Buffer->GetData(), PointsToRead * Header.PointDataRecordLength);
+
+					ThreadResults.Add(Async(EAsyncExecution::ThreadPool, [PointsToRead, Buffer, &ImportScale, &bHasIntensityData, &IntensityMultiplier, &RGBMultiplier, &Header, &OutImportResults, &Settings, &bFirstPointSet, &CoordsLock, &PointsLock]
+						{
+							uint8* Data = Buffer->GetData();
+
+							TArray<FLidarPointCloudPoint> Points;
+							Points.Reserve(PointsToRead);
+
+							FBox _Bounds(EForceInit::ForceInit);
+
+							TArray<uint8> Classifications;
+
+							// Parse the data
+							for (int64 j = 0; j < PointsToRead && !OutImportResults.IsCancelled(); j++)
+							{
+								FLidarPointCloudFileIO_LAS_PointDataRecordFormatCommon* Record = (FLidarPointCloudFileIO_LAS_PointDataRecordFormatCommon*)Data;
+
+								FVector ProcessedLocation;
+								float Intensity = bHasIntensityData ? Record->Intensity * IntensityMultiplier : 1.0f;
+								float R = 1.0f;
+								float G = 1.0f;
+								float B = 1.0f;
+								uint8 Classification = Header.IsLegacyFormat() ? ((FLidarPointCloudFileIO_LAS_PointDataRecordFormat0*)Data)->Classification : ((FLidarPointCloudFileIO_LAS_PointDataRecordFormat6*)Data)->Classification;
+
+								Classifications.AddUnique(Classification);
+
+								// Extract location information
+								{
+									// Calculate the actual location of the point, convert to UU and flip the Y axis
+									FDoubleVector Location = (Header.ScaleFactor * Record->Location + Header.Offset) * ImportScale;
+									Location.Y = -Location.Y;
+
+									if (!bFirstPointSet)
+									{
+										FScopeLock Lock(&CoordsLock);
+
+										if (!bFirstPointSet)
+										{
+											OutImportResults.OriginalCoordinates = Location;
+											bFirstPointSet = true;
+										}
+									}
+
+									// Shift to protect from precision loss
+									Location -= OutImportResults.OriginalCoordinates;
+
+									_Bounds += Location.ToVector();
+
+									// Convert location to floats
+									ProcessedLocation = Location.ToVector();
+								}
+
+								// Extract color information
+								{
+									FLidarPointCloudFileIO_LAS_PointDataRecordFormatCommonRGB* RecordRGB = nullptr;
+
+									switch (Header.PointDataRecordFormat)
+									{
+									case 2: RecordRGB = (FLidarPointCloudFileIO_LAS_PointDataRecordFormat2*)Data; break;
+									case 3: RecordRGB = (FLidarPointCloudFileIO_LAS_PointDataRecordFormat3*)Data; break;
+									case 5: RecordRGB = (FLidarPointCloudFileIO_LAS_PointDataRecordFormat5*)Data; break;
+									case 7: RecordRGB = (FLidarPointCloudFileIO_LAS_PointDataRecordFormat7*)Data; break;
+									case 8: RecordRGB = (FLidarPointCloudFileIO_LAS_PointDataRecordFormat8*)Data; break;
+									case 10: RecordRGB = (FLidarPointCloudFileIO_LAS_PointDataRecordFormat10*)Data; break;
+									}
+
+									if (RecordRGB)
+									{
+										R = RecordRGB->Red * RGBMultiplier;
+										G = RecordRGB->Green * RGBMultiplier;
+										B = RecordRGB->Blue * RGBMultiplier;
+									}
+								}
+
+								Points.Emplace(ProcessedLocation, R, G, B, Intensity, Classification);
+
+								// Increment the data pointer
+								Data += Header.PointDataRecordLength;
+							}
+
+							// Data Sync
+							{
+								FScopeLock Lock(&PointsLock);
+
+								OutImportResults.AddPointsBulk(Points);
+								OutImportResults.Bounds += _Bounds;
+
+								for (uint8& Classification : Classifications)
+								{
+									OutImportResults.ClassificationsImported.AddUnique(Classification);
+								}
+							}
+
+							// Memory Sync
+							Buffer->bInUse = false;
+						}));
+
+					PointsRead += PointsToRead;
+				}
+
+				// Sync threads
+				for (auto& ThreadResult : ThreadResults) ThreadResult.Get();
+
+				// Apply re-centering, if selected and not applied already
+				if (bAutoCenter)
+				{
+					OutImportResults.CenterPoints();
+				}
+
+				// Make sure to progress the counter to the end before returning
+				OutImportResults.IncrementProgressCounter(TotalPointsToRead);
+			}
+			
+			bSuccess = !OutImportResults.IsCancelled();
+		}
+
+		// Free memory
+		Reader->Close();
+	}
+
+	return bSuccess;
+}
+
+bool ULidarPointCloudFileIO_LAS::HandleExport(const FString& Filename, ULidarPointCloud* PointCloud)
+{
+	FArchive* Ar = IFileManager::Get().CreateFileWriter(*Filename, 0);
+	if (!Ar)
+	{
+		return false;
+	}
+
+	FDoubleVector Min = PointCloud->GetBounds().Min;
+	FDoubleVector Max = PointCloud->GetBounds().Max;
+
+	// Add original coordinates
+	Min += PointCloud->OriginalCoordinates;
+	Max += PointCloud->OriginalCoordinates;
+
+	// Flip Y
+	float MaxY = Max.Y;
+	Max.Y = -Min.Y;
+	Min.Y = -MaxY;
+
+	const float ExportScale = GetDefault<ULidarPointCloudSettings>()->ExportScale;
+
+	// Convert to meters
+	Min *= ExportScale;
+	Max *= ExportScale;
+
+	FLidarPointCloudFileIO_LAS_PublicHeaderBlock Header = FLidarPointCloudFileIO_LAS_PublicHeaderBlock::Generate(PointCloud->GetNumPoints(), Min, Max);
+
+	Ar->Serialize(&Header, Header.HeaderSize);
+	
+	FDoubleVector Size = Max - Min;
+	FDoubleVector ForwardScale(FMath::Pow(2, 31 - FMath::CeilToInt(FMath::Log2(Size.X))), FMath::Pow(2, 31 - FMath::CeilToInt(FMath::Log2(Size.Y))), FMath::Pow(2, 31 - FMath::CeilToInt(FMath::Log2(Size.Z))));
+
+	int64 NumProcessedPoints = 0;
+	TArray<FLidarPointCloudPoint*> Points;
+
+	int64 MaxBatchSize = GetDefault<ULidarPointCloudSettings>()->ExportBatchSize;
+
+	FLidarPointCloudFileIO_LAS_PointDataRecordFormat2* PointRecord = new FLidarPointCloudFileIO_LAS_PointDataRecordFormat2();
+	FMemory::Memzero(PointRecord, Header.PointDataRecordLength);
+
+	while (NumProcessedPoints < PointCloud->GetNumPoints())
+	{
+		int64 BatchSize = FMath::Min(MaxBatchSize, PointCloud->GetNumPoints() - NumProcessedPoints);
+		PointCloud->GetPoints(Points, NumProcessedPoints, BatchSize);
+
+		for (FLidarPointCloudPoint* Point : Points)
+		{
+			FDoubleVector Location = (PointCloud->OriginalCoordinates + Point->Location) * ExportScale;
+			Location.Y = -Location.Y;
+
+			PointRecord->Location = (ForwardScale * (Location - Min)).ToIntVector();
+			PointRecord->Intensity = (Point->Color.A << 8) + Point->Color.A;
+			PointRecord->Red = (Point->Color.R << 8) + Point->Color.R;
+			PointRecord->Green = (Point->Color.G << 8) + Point->Color.G;
+			PointRecord->Blue = (Point->Color.B << 8) + Point->Color.B;
+			PointRecord->Classification = Point->ClassificationID;
+
+			Ar->Serialize(PointRecord, Header.PointDataRecordLength);
+		}
+
+		NumProcessedPoints += BatchSize;
+	}
+
+	delete PointRecord;
+	delete Ar;
+	return true;
+}
diff --git a/Engine/Plugins/Enterprise/LidarPointCloud/Source/LidarPointCloudRuntime/Private/LidarPointCloud.cpp b/Engine/Plugins/Enterprise/LidarPointCloud/Source/LidarPointCloudRuntime/Private/LidarPointCloud.cpp
new file mode 100644
--- /dev/null
+++ b/Engine/Plugins/Enterprise/LidarPointCloud/Source/LidarPointCloudRuntime/Private/LidarPointCloud.cpp
@@ -0,0 +1,1108 @@
+// Copyright Epic Games, Inc. All Rights Reserved.
+
+#include "LidarPointCloud.h"
+#include "LidarPointCloudShared.h"
+#include "LidarPointCloudActor.h"
+#include "LidarPointCloudComponent.h"
+#include "IO/LidarPointCloudFileIO.h"
+#include "Async/Async.h"
+#include "Async/Future.h"
+#include "Serialization/CustomVersion.h"
+#include "Misc/ScopeTryLock.h"
+#include "Misc/ConfigCacheIni.h"
+#include "Engine/Engine.h"
+#include "LatentActions.h"
+#include "PhysicsEngine/BodySetup.h"
+#include "Framework/Notifications/NotificationManager.h"
+#include "Engine/CollisionProfile.h"
+#include "EngineUtils.h"
+
+#if WITH_EDITOR
+#include "IContentBrowserSingleton.h"
+#include "ContentBrowserModule.h"
+#include "AssetRegistryModule.h"
+#include "Editor.h"
+#include "Styling/SlateStyleRegistry.h"
+#include "Misc/MessageDialog.h"
+#include "UObject/UObjectGlobals.h"
+#endif
+
+#define IS_PROPERTY(Name) PropertyChangedEvent.MemberProperty->GetName().Equals(#Name)
+
+const FGuid ULidarPointCloud::PointCloudFileGUID('P', 'C', 'P', 'F');
+const int32 ULidarPointCloud::PointCloudFileVersion(15);
+FCustomVersionRegistration PCPFileVersion(ULidarPointCloud::PointCloudFileGUID, ULidarPointCloud::PointCloudFileVersion, TEXT("LiDAR Point Cloud File Version"));
+
+#define LOCTEXT_NAMESPACE "LidarPointCloud"
+
+/////////////////////////////////////////////////
+// FLidarPointCloudNotification
+
+FLidarPointCloudNotification::FLidarPointCloudNotification(UObject* Owner)
+	: Owner(Owner)
+	, CurrentText("")
+	, CurrentProgress(-1)
+{
+}
+
+void FLidarPointCloudNotification::Create(const FString& Text, FThreadSafeBool* bCancelPtr, const FString& Icon)
+{
+	SetTextWithProgress(Text, -1);
+
+#if WITH_EDITOR
+	if (Owner && !IsValid() && GIsEditor)
+	{
+		// Build the notification widget
+		FNotificationInfo Info(FText::FromString(CurrentText));
+		Info.bFireAndForget = false;
+		Info.Image = FSlateStyleRegistry::FindSlateStyle("LidarPointCloudStyle")->GetBrush(*Icon);
+
+		if (Owner->HasAnyFlags(RF_Public | RF_Standalone))
+		{
+			Info.Hyperlink = FSimpleDelegate::CreateLambda([this] {
+				// Select the cloud in Content Browser when the hyperlink is clicked
+				TArray<FAssetData> AssetData;
+				AssetData.Add(FModuleManager::LoadModuleChecked<FAssetRegistryModule>("AssetRegistry").Get().GetAssetByObjectPath(FSoftObjectPath(Owner).GetAssetPathName()));
+				FModuleManager::LoadModuleChecked<FContentBrowserModule>("ContentBrowser").Get().SyncBrowserToAssets(AssetData);
+			});
+			Info.HyperlinkText = FText::FromString(FPaths::GetBaseFilename(FSoftObjectPath(Owner).ToString()));
+		}
+
+		if (bCancelPtr)
+		{
+			Info.ButtonDetails.Emplace(
+				LOCTEXT("OpCancel", "Cancel"),
+				LOCTEXT("OpCancelToolTip", "Cancels the point cloud operation in progress."),
+				FSimpleDelegate::CreateLambda([bCancelPtr] { *bCancelPtr = true; })
+			);
+		}
+
+		NotificationItem = FSlateNotificationManager::Get().AddNotification(Info);
+		if (IsValid())
+		{
+			NotificationItem->SetCompletionState(SNotificationItem::CS_Pending);
+		}
+	}
+#endif
+}
+
+void FLidarPointCloudNotification::SetText(const FString& Text)
+{
+	CurrentText = Text;
+	UpdateStatus();
+}
+
+void FLidarPointCloudNotification::SetProgress(int8 Progress)
+{
+	CurrentProgress = Progress;
+	UpdateStatus();
+}
+
+void FLidarPointCloudNotification::SetTextWithProgress(const FString& Text, int8 Progress)
+{
+	CurrentText = Text;
+	CurrentProgress = Progress;
+	UpdateStatus();
+}
+
+void FLidarPointCloudNotification::Close(bool bSuccess)
+{
+#if WITH_EDITOR
+	if (Owner && IsValid())
+	{
+		// Do not use fadeout if the engine is shutting down
+		if (!GEditor->HasAnyFlags(RF_BeginDestroyed))
+		{
+			CurrentText.Append(bSuccess ? " Complete" : " Failed");
+			CurrentProgress = -1;
+			UpdateStatus();
+			NotificationItem->SetCompletionState(bSuccess ? SNotificationItem::CS_Success : SNotificationItem::CS_Fail);
+			NotificationItem->ExpireAndFadeout();
+		}
+		NotificationItem.Reset();
+	}
+#endif
+}
+
+void FLidarPointCloudNotification::UpdateStatus()
+{
+	if (!Owner || !IsValid())
+	{
+		return;
+	}
+
+	if (IsInGameThread())
+	{
+		// Update Text
+		{
+			FString Message;
+
+			if (CurrentProgress >= 0)
+			{
+				Message = FString::Printf(TEXT("%s: %d%%"), *CurrentText, CurrentProgress);
+			}
+			else
+			{
+				Message = FString::Printf(TEXT("%s"), *CurrentText);
+			}
+
+			NotificationItem->SetText(FText::FromString(Message));
+		}
+	}
+	else
+	{
+		AsyncTask(ENamedThreads::GameThread, [this] { UpdateStatus(); });
+	}
+}
+
+/////////////////////////////////////////////////
+// ULidarPointCloud
+
+ULidarPointCloud::ULidarPointCloud()
+	: CollisionAccuracy(100)
+	, Octree()
+	, OriginalCoordinates(FDoubleVector::ZeroVector)
+	, Notification(this)
+	, BodySetup(nullptr)
+	, bCollisionBuildInProgress(false)
+{
+	// Make sure we are transactional to allow undo redo
+	this->SetFlags(RF_Transactional);
+}
+
+void ULidarPointCloud::Serialize(FArchive& Ar)
+{
+	DECLARE_SCOPE_CYCLE_COUNTER(TEXT("ULidarPointCloud::Serialize"), STAT_PointCLoud_Serialize, STATGROUP_LoadTime);
+
+	Ar.UsingCustomVersion(PointCloudFileGUID);
+
+	Super::Serialize(Ar);
+		
+	int32 Version = Ar.CustomVer(PointCloudFileGUID);
+
+	if (Version > 13)
+	{
+		Ar << BodySetup;
+
+		if (Ar.IsCountingMemory())
+		{
+			if (BodySetup)
+			{
+				BodySetup->Serialize(Ar);
+			}
+		}
+	}
+
+	// Make sure to serialize only actual data
+	if (Ar.ShouldSkipBulkData() || Ar.IsObjectReferenceCollector() || !Ar.IsPersistent())
+	{
+		return;
+	}
+
+	FScopeBenchmarkTimer BenchmarkTimer("Serialization");
+	
+	ULidarPointCloudFileIO::SerializeImportSettings(Ar, ImportSettings);
+
+	// Do not save the Octree, if in the middle of an import or the access to the data is blocked
+	{
+		FScopeTryLock LockImport(&ImportLock);
+		FScopeTryLock LockOctree(&Octree.DataLock);
+
+		bool bValidOctree = LockImport.IsLocked() && LockOctree.IsLocked();
+		Ar << bValidOctree;
+		if (bValidOctree)
+		{
+			Ar << Octree;
+		}
+	}
+}
+
+void ULidarPointCloud::PostLoad()
+{
+	Super::PostLoad();
+
+	InitializeCollisionRendering();
+}
+
+void ULidarPointCloud::GetAssetRegistryTags(TArray<FAssetRegistryTag>& OutTags) const
+{
+	OutTags.Add(FAssetRegistryTag("PointCount", PointCloudAssetRegistryCache.PointCount, FAssetRegistryTag::TT_Numerical));
+	OutTags.Add(FAssetRegistryTag("ApproxSize", PointCloudAssetRegistryCache.ApproxSize, FAssetRegistryTag::TT_Dimensional));
+
+	Super::GetAssetRegistryTags(OutTags);
+}
+
+void ULidarPointCloud::BeginDestroy()
+{
+	Super::BeginDestroy();
+
+	// Cancel async import and wait for it to exit
+	bAsyncCancelled = true;
+	FScopeLock LockImport(&ImportLock);
+
+	// Hide any notifications, if still present
+	Notification.Close(false);
+
+	// Wait for ongoing data access to finish
+	FScopeLock LockOctree(&Octree.DataLock);
+
+	// Release any collision rendering data, if present
+	ReleaseCollisionRendering();
+}
+
+#if WITH_EDITOR
+void ULidarPointCloud::PostEditChangeProperty(FPropertyChangedEvent& PropertyChangedEvent)
+{
+	if (PropertyChangedEvent.MemberProperty)
+	{		
+		if (IS_PROPERTY(SourcePath))
+		{
+			SetSourcePath(SourcePath.FilePath);
+		}
+
+		if (IS_PROPERTY(CollisionAccuracy))
+		{
+			if (CollisionAccuracy < Octree.GetEstimatedPointSpacing())
+			{
+				FMessageDialog::Open(EAppMsgType::Type::Ok, FText::FromString(FString::Printf(TEXT("Average point spacing is estimated to be around %f cm.\nSetting accuracy close to or lower than that value may result in collision holes."), FMath::RoundToFloat(Octree.GetEstimatedPointSpacing() * 100) * 0.01f)));
+			}
+		}
+	}
+
+	Super::PostEditChangeProperty(PropertyChangedEvent);
+}
+#endif
+
+int32 ULidarPointCloud::GetDataSize() const
+{
+	const int64 OctreeSize = Octree.GetAllocatedSize();
+	const int64 CollisionSize = Octree.GetCollisionData()->Indices.GetAllocatedSize() + Octree.GetCollisionData()->Vertices.GetAllocatedSize();
+
+	return (OctreeSize + CollisionSize) >> 20;
+}
+
+bool ULidarPointCloud::HasCollisionData() const
+{
+	return Octree.HasCollisionData();
+}
+
+TArray<FLidarPointCloudPoint*> ULidarPointCloud::GetPoints(int64 StartIndex, int64 Count)
+{
+	TArray<FLidarPointCloudPoint*> Points;
+	GetPoints(Points, StartIndex, Count);
+	return Points;
+}
+
+TArray<FLidarPointCloudPoint> ULidarPointCloud::GetPointsAsCopies(int32 StartIndex, int32 Count) const
+{
+	TArray<FLidarPointCloudPoint> Points;
+	GetPointsAsCopies(Points, StartIndex, Count);
+	return Points;
+}
+
+TArray<FLidarPointCloudPoint> ULidarPointCloud::GetPointsInSphereAsCopies(FVector Center, float Radius, bool bVisibleOnly)
+{
+	TArray<FLidarPointCloudPoint> Points;
+	GetPointsInSphereAsCopies(Points, FSphere(Center, Radius), bVisibleOnly);
+	return Points;
+}
+
+TArray<FLidarPointCloudPoint> ULidarPointCloud::GetPointsInBoxAsCopies(FVector Center, FVector Extent, bool bVisibleOnly)
+{
+	TArray<FLidarPointCloudPoint> Points;
+	GetPointsInBoxAsCopies(Points, FBox(Center - Extent, Center + Extent), bVisibleOnly);
+	return Points;
+}
+
+bool ULidarPointCloud::LineTraceSingle(FVector Origin, FVector Direction, float Radius, bool bVisibleOnly, FLidarPointCloudPoint& PointHit)
+{
+	FLidarPointCloudPoint* Point = LineTraceSingle(FLidarPointCloudRay(Origin, Direction), Radius, bVisibleOnly);
+	if (Point)
+	{
+		PointHit = *Point;
+		return true;
+	}
+
+	return false;
+}
+
+void ULidarPointCloud::SetSourcePath(const FString& NewSourcePath)
+{
+	SourcePath.FilePath = NewSourcePath;
+
+	if (FPaths::FileExists(SourcePath.FilePath))
+	{
+		if (FPaths::IsRelative(SourcePath.FilePath))
+		{
+			SourcePath.FilePath = FPaths::ConvertRelativePathToFull(SourcePath.FilePath);
+		}
+
+		// Generate new ImportSettings if the source path has changed
+		ImportSettings = ULidarPointCloudFileIO::GetImportSettings(SourcePath.FilePath);
+	}
+	else
+	{
+		// Invalidate ImportSettings if the source path is invalid too
+		ImportSettings = nullptr;
+	}
+}
+
+void ULidarPointCloud::BuildCollision()
+{
+	if (bCollisionBuildInProgress)
+	{
+		PC_ERROR("Another collision operation already in progress.");
+		return;
+	}
+
+	Notification.Create("Building Collision", nullptr, "LidarPointCloudEditor.BuildCollision");
+	
+	bCollisionBuildInProgress = true;
+	MarkPackageDirty();
+
+	UBodySetup* NewBodySetup = NewObject<UBodySetup>(this);
+	NewBodySetup->BodySetupGuid = FGuid::NewGuid();
+	NewBodySetup->CollisionTraceFlag = CTF_UseComplexAsSimple;
+	NewBodySetup->bHasCookedCollisionData = true;
+
+	Async(EAsyncExecution::Thread, [this, NewBodySetup]{
+		Octree.BuildCollision(CollisionAccuracy, true);
+
+		FBenchmarkTimer::Reset();
+#if WITH_PHYSX  && PHYSICS_INTERFACE_PHYSX
+		AsyncTask(ENamedThreads::GameThread, [this, NewBodySetup] {
+			NewBodySetup->CreatePhysicsMeshesAsync(FOnAsyncPhysicsCookFinished::CreateUObject(this, &ULidarPointCloud::FinishPhysicsAsyncCook, NewBodySetup));
+		});
+#elif WITH_CHAOS
+		NewBodySetup->CreatePhysicsMeshes();
+		AsyncTask(ENamedThreads::GameThread, [this, NewBodySetup] { FinishPhysicsAsyncCook(true, NewBodySetup); });
+#endif		
+	});
+}
+
+void ULidarPointCloud::RemoveCollision()
+{
+	if (bCollisionBuildInProgress)
+	{
+		PC_ERROR("Another collision operation already in progress.");
+		return;
+	}
+
+	bCollisionBuildInProgress = true;
+	
+	MarkPackageDirty();
+
+	Octree.RemoveCollision();
+
+	BodySetup = NewObject<UBodySetup>(this);
+	ReleaseCollisionRendering();
+	InitializeCollisionRendering();
+	OnPointCloudUpdateCollisionEvent.Broadcast();
+
+	bCollisionBuildInProgress = false;
+}
+
+void ULidarPointCloud::CenterPoints()
+{
+	// Only shift, if not already centered
+	if (!IsCentered())
+	{
+		ShiftPointsBy(-GetPointsBounds().GetCenter(), true);
+	}
+}
+
+void ULidarPointCloud::RestoreOriginalCoordinates()
+{
+	if (!OriginalCoordinates.IsNearlyZero(0.001f))
+	{
+		ShiftPointsBy(OriginalCoordinates, true);
+	}
+}
+
+void ULidarPointCloud::ShiftPointsBy(FDoubleVector Offset, bool bRefreshPointsBounds)
+{
+	Octree.ShiftPointsBy(Offset, bRefreshPointsBounds);
+	OriginalCoordinates -= Offset;
+	MarkPackageDirty();
+	OnPointCloudRebuiltEvent.Broadcast();
+}
+
+bool ULidarPointCloud::IsCentered() const
+{
+	// Return as centered if within 1mm of center
+	return GetPointsBounds().GetCenter().IsNearlyZero(0.1f);
+}
+
+void ULidarPointCloud::Reimport(const FLidarPointCloudAsyncParameters& AsyncParameters)
+{
+	if (FPaths::FileExists(SourcePath.FilePath))
+	{
+		FScopeTryLock Lock(&ImportLock);
+
+		if (!Lock.IsLocked())
+		{
+			PC_ERROR("Cannot reimport the asset - data is currently being used.");
+			return;
+		}
+
+		bAsyncCancelled = false;
+		Notification.Create("Importing Point Cloud", &bAsyncCancelled);
+
+		// The actual import function to be executed
+		auto ImportFunction = [this, AsyncParameters]
+		{
+			// This will take over the lock
+			FScopeLock Lock(&ImportLock);
+
+			FLidarPointCloudImportResults ImportResults(&bAsyncCancelled, [this, AsyncParameters](float Progress)
+			{
+				Notification.SetProgress(50.0f * Progress);
+				if (AsyncParameters.ProgressCallback)
+				{
+					AsyncParameters.ProgressCallback(50.0f * Progress);
+				}
+			});
+
+			bool bSuccess = false;
+
+			if (ULidarPointCloudFileIO::Import(SourcePath.FilePath, ImportSettings, ImportResults))
+			{
+				ClassificationsImported = ImportResults.ClassificationsImported;
+
+				// Wait for rendering to complete before proceeding and lock the access to the data
+				FScopeLock DataLock(&Octree.DataLock);
+				
+				// Re-initialize the Octree
+				Octree.Initialize(ImportResults.Bounds);
+
+				FScopeBenchmarkTimer BenchmarkTimer("Octree Build-Up");
+
+				bSuccess = InsertPoints_NoLock(ImportResults.Points.GetData(), ImportResults.Points.Num(), GetDefault<ULidarPointCloudSettings>()->DuplicateHandling, false, &bAsyncCancelled, [this, AsyncParameters](float Progress)
+				{
+					Notification.SetProgress(50.0f + 50.0f * Progress);
+					if (AsyncParameters.ProgressCallback)
+					{
+						AsyncParameters.ProgressCallback(50.0f + 50.0f * Progress);
+					}
+				});
+
+				if (bSuccess)
+				{
+					RefreshPointsBounds();
+					OriginalCoordinates = ImportResults.OriginalCoordinates;
+				}
+				else
+				{
+					BenchmarkTimer.bActive = false;
+
+					Octree.Empty(true);
+
+					// Update PointCloudAssetRegistryCache
+					PointCloudAssetRegistryCache.PointCount = FString::FromInt(Octree.GetNumPoints());
+				}	
+			}
+
+			// Only process those if not being destroyed
+			if (!HasAnyFlags(RF_BeginDestroyed))
+			{
+				auto PostFunction = [this, bSuccess]() {
+					MarkPackageDirty();
+					Notification.Close(bSuccess);
+					OnPointCloudRebuiltEvent.Broadcast();
+				};
+
+				// Make sure the call is executed on the correct thread if using async
+				if (IsInGameThread())
+				{
+					PostFunction();
+				}
+				else
+				{
+					AsyncTask(ENamedThreads::GameThread, MoveTemp(PostFunction));
+				}
+			}	
+			
+			if (AsyncParameters.CompletionCallback)
+			{
+				AsyncParameters.CompletionCallback(bSuccess);
+			}
+
+			if(!bSuccess)
+			{
+				PC_ERROR("Point Cloud importing failed or cancelled.");
+			}
+		};
+
+		if (AsyncParameters.bUseAsync)
+		{
+			Async(EAsyncExecution::Thread, MoveTemp(ImportFunction));
+		}
+		else
+		{
+			ImportFunction();
+		}
+	}
+	else
+	{
+		PC_ERROR("Reimport failed, provided source path '%s' could not be found.", *SourcePath.FilePath);
+
+		if (AsyncParameters.CompletionCallback)
+		{
+			AsyncParameters.CompletionCallback(false);
+		}
+	}
+}
+
+bool ULidarPointCloud::Export(const FString& Filename)
+{
+	return ULidarPointCloudFileIO::Export(Filename, this);
+}
+
+void ULidarPointCloud::InsertPoint(const FLidarPointCloudPoint& Point, ELidarPointCloudDuplicateHandling DuplicateHandling, bool bRefreshPointsBounds)
+{
+	FScopeLock Lock(&Octree.DataLock);
+
+	Octree.InsertPoint(&Point, DuplicateHandling, bRefreshPointsBounds);
+
+	// Update PointCloudAssetRegistryCache
+	PointCloudAssetRegistryCache.PointCount = FString::FromInt(Octree.GetNumPoints());
+}
+
+template<typename T>
+bool ULidarPointCloud::InsertPoints_NoLock(T InPoints, const int64& Count, ELidarPointCloudDuplicateHandling DuplicateHandling, bool bRefreshPointsBounds, FThreadSafeBool* bCanceled, TFunction<void(float)> ProgressCallback)
+{
+	const ULidarPointCloudSettings* Settings = GetDefault<ULidarPointCloudSettings>();
+	const int32 MaxBatchSize = Settings->MultithreadingInsertionBatchSize;
+	
+	if (Settings->bUseMultithreading)
+	{
+		// Minimum amount of points to progress to count as 1%
+		int64 RefreshStatusFrequency = Count * 0.01f;
+		FThreadSafeCounter64 ProcessedPoints(0);
+		int64 TotalProcessedPoints = 0;
+
+		const int32 NumThreads = FMath::Min(FPlatformMisc::NumberOfCoresIncludingHyperthreads() - 1, (int32)(Count / MaxBatchSize) + 1);
+		TFuture<void>* ThreadResults = new TFuture<void>[NumThreads];
+		const int64 NumPointsPerThread = Count / NumThreads + 1;
+
+		FCriticalSection ProgressCallbackLock;
+
+		// Fire threads
+		for (int32 t = 0; t < NumThreads; t++)
+		{
+			const int32 ThreadID = t;
+
+			ThreadResults[ThreadID] = Async(EAsyncExecution::Thread, [this, ThreadID, DuplicateHandling, bRefreshPointsBounds, MaxBatchSize, NumPointsPerThread, RefreshStatusFrequency, &ProcessedPoints, &TotalProcessedPoints, InPoints, Count, &ProgressCallback, &ProgressCallbackLock, &bCanceled]
+				{
+					int64 Idx = ThreadID * NumPointsPerThread;
+					int64 MaxIdx = FMath::Min(Idx + NumPointsPerThread, Count);
+					T DataPointer = InPoints + Idx;
+
+					while (Idx < MaxIdx)
+					{
+						int32 BatchSize = FMath::Min(MaxIdx - Idx, (int64)MaxBatchSize);
+
+						Octree.InsertPoints(DataPointer, BatchSize, DuplicateHandling, bRefreshPointsBounds);
+
+						if (ProgressCallback)
+						{
+							ProcessedPoints.Add(BatchSize);
+							if (ProcessedPoints.GetValue() > RefreshStatusFrequency)
+							{
+								FScopeLock Lock(&ProgressCallbackLock);
+								TotalProcessedPoints += ProcessedPoints.GetValue();
+								ProcessedPoints.Reset();
+								ProgressCallback((double)TotalProcessedPoints / Count);
+							}
+						}
+
+						if (bCanceled && *bCanceled)
+						{
+							break;
+						}
+
+						Idx += BatchSize;
+						DataPointer += BatchSize;
+					}
+				});
+		}
+
+		// Sync
+		for (int32 i = 0; i < NumThreads; i++) ThreadResults[i].Get();
+
+		// Cleanup
+		delete[] ThreadResults;
+	}
+	else
+	{
+		int64 Idx = 0;
+		T DataPointer = InPoints;
+
+		while (Idx < Count)
+		{
+			int32 BatchSize = FMath::Min(Count - Idx, (int64)MaxBatchSize);
+
+			Octree.InsertPoints(DataPointer, BatchSize, DuplicateHandling, bRefreshPointsBounds);
+
+			Idx += BatchSize;
+			DataPointer += BatchSize;
+		}
+	}
+
+	// Do not attempt to touch Render Data if being destroyed
+	if (!HasAnyFlags(RF_BeginDestroyed))
+	{
+		// Update PointCloudAssetRegistryCache
+		PointCloudAssetRegistryCache.PointCount = FString::FromInt(Octree.GetNumPoints());
+	}
+
+	return !bCanceled || !(*bCanceled);
+}
+
+template<typename T>
+bool ULidarPointCloud::SetData(T Points, const int64& Count, TFunction<void(float)> ProgressCallback)
+{
+	// Lock the point cloud
+	FScopeLock Lock(&ImportLock);
+
+	// Calculate the bounds
+	FBox Bounds = CalculateBoundsFromPoints(Points, Count);
+
+	bool bSuccess = false;
+
+	// Only proceed if the bounds are valid
+	if (Bounds.IsValid)
+	{
+		// Wait for rendering to complete before proceeding and lock the access to the data
+		FScopeLock DataLock(&Octree.DataLock);
+
+		// Initialize the Octree
+		Initialize(Bounds);
+
+		bSuccess = InsertPoints_NoLock(Points, Count, GetDefault<ULidarPointCloudSettings>()->DuplicateHandling, false, nullptr, MoveTemp(ProgressCallback));
+
+		if (bSuccess)
+		{
+			RefreshPointsBounds();
+		}
+		else
+		{
+			Octree.Empty(true);
+		}
+
+		// Only process those if not being destroyed
+		if (!HasAnyFlags(RF_BeginDestroyed))
+		{
+			auto PostFunction = [this, bSuccess]() {
+				MarkPackageDirty();
+				Notification.Close(bSuccess);
+				OnPointCloudRebuiltEvent.Broadcast();
+			};
+
+			// Make sure the call is executed on the correct thread if using async
+			if (IsInGameThread())
+			{
+				PostFunction();
+			}
+			else
+			{
+				AsyncTask(ENamedThreads::GameThread, MoveTemp(PostFunction));
+			}
+		}
+	}
+
+	if (!bSuccess)
+	{
+		PC_ERROR("Setting Point Cloud data failed.");
+	}
+
+	return bSuccess;
+}
+
+void ULidarPointCloud::Merge(TArray<ULidarPointCloud*> PointCloudsToMerge, TFunction<void(float)> ProgressCallback)
+{
+	for (int32 i = 0; i < PointCloudsToMerge.Num(); i++)
+	{
+		if (!IsValid(PointCloudsToMerge[i]) || PointCloudsToMerge[i] == this)
+		{
+			PointCloudsToMerge.RemoveAtSwap(i--, 1, false);
+		}
+	}
+
+	PointCloudsToMerge.Shrink();
+
+	// Abort if no valid assets are found
+	if (PointCloudsToMerge.Num() == 0)
+	{
+		return;
+	}
+
+	FScopeBenchmarkTimer Timer("Merge");
+
+	// Lock the point cloud
+	FScopeLock Lock(&ImportLock);
+	FScopeLock DataLock(&Octree.DataLock);
+
+	// Calculate new, combined bounds
+	FBox NewBounds = GetBounds();
+	for (ULidarPointCloud* Asset : PointCloudsToMerge)
+	{
+		NewBounds += Asset->GetBounds();
+
+		for (uint8& Classification : Asset->ClassificationsImported)
+		{
+			ClassificationsImported.AddUnique(Classification);
+		}
+	}
+
+	// Make a copy of current points, as the data will be reinitialized
+	TArray<FLidarPointCloudPoint> Points;
+	GetPointsAsCopies(Points);
+
+	// Initialize the Octree
+	Initialize(NewBounds);
+
+	if (ProgressCallback)
+	{
+		ProgressCallback(0);
+	}
+
+	// Re-insert original points
+	InsertPoints(Points, GetDefault<ULidarPointCloudSettings>()->DuplicateHandling, false);
+
+	int32 Progress = 0;
+	float ProgressMulti = 1.0f / (1 + PointCloudsToMerge.Num());
+
+	if (ProgressCallback)
+	{
+		ProgressCallback(++Progress * ProgressMulti);
+	}
+
+	// Insert other points
+	for (ULidarPointCloud* Asset : PointCloudsToMerge)
+	{
+		// #todo: Do this in batches to avoid spiking RAM
+		Asset->GetPointsAsCopies(Points);
+		InsertPoints(Points, GetDefault<ULidarPointCloudSettings>()->DuplicateHandling, false);
+
+		if (ProgressCallback)
+		{
+			ProgressCallback(++Progress * ProgressMulti);
+		}
+	}
+
+	RefreshPointsBounds();
+
+	auto PostFunction = [this]() {
+		MarkPackageDirty();
+		OnPointCloudRebuiltEvent.Broadcast();
+	};
+
+	// Make sure the call is executed on the correct thread if using async
+	if (IsInGameThread())
+	{
+		PostFunction();
+	}
+	else
+	{
+		AsyncTask(ENamedThreads::GameThread, MoveTemp(PostFunction));
+	}
+}
+
+bool ULidarPointCloud::GetPhysicsTriMeshData(FTriMeshCollisionData* CollisionData, bool InUseAllTriData)
+{
+	CollisionData->Vertices = Octree.GetCollisionData()->Vertices;
+	CollisionData->Indices = Octree.GetCollisionData()->Indices;
+
+	return true;
+}
+
+UBodySetup* ULidarPointCloud::GetBodySetup()
+{
+	return BodySetup && !BodySetup->IsPendingKill() ? BodySetup : nullptr;
+}
+
+void ULidarPointCloud::AlignClouds(TArray<ULidarPointCloud*> PointCloudsToAlign)
+{
+	FDoubleVector Min(FLT_MAX);
+	FDoubleVector Max(-FLT_MAX);
+
+	// Calculate combined bounds
+	for (ULidarPointCloud* Asset : PointCloudsToAlign)
+	{
+		FDoubleVector OffsetMin = Asset->OriginalCoordinates + Asset->GetBounds().Min;
+		FDoubleVector OffsetMax = Asset->OriginalCoordinates + Asset->GetBounds().Max;
+
+		Min.X = FMath::Min(Min.X, OffsetMin.X);
+		Min.Y = FMath::Min(Min.Y, OffsetMin.Y);
+		Min.Z = FMath::Min(Min.Z, OffsetMin.Z);
+
+		Max.X = FMath::Max(Max.X, OffsetMax.X);
+		Max.Y = FMath::Max(Max.Y, OffsetMax.Y);
+		Max.Z = FMath::Max(Max.Z, OffsetMax.Z);
+	}
+
+	FDoubleVector CombinedCenter = (Min + Max) / 2;
+
+	// Calculate and apply individual shifts
+	for (ULidarPointCloud* Asset : PointCloudsToAlign)
+	{
+		FDoubleVector OffsetCenter = Asset->OriginalCoordinates + Asset->GetBounds().GetCenter();
+		FDoubleVector Delta = OffsetCenter - CombinedCenter;
+		Asset->ShiftPointsBy(Delta, true);
+	}
+}
+
+ULidarPointCloud* ULidarPointCloud::CreateFromFile(const FString& Filename, const FLidarPointCloudAsyncParameters& AsyncParameters, TSharedPtr<struct FLidarPointCloudImportSettings> ImportSettings, UObject* InParent, FName InName, EObjectFlags Flags)
+{
+#if WITH_EDITOR
+	FOnPointCloudChanged OnPointCloudRebuiltEvent;
+	FOnPointCloudChanged OnPointCloudUpdateCollisionEvent;
+	bool bOldPointCloudExists = false;
+
+	// See if Point Cloud already exists
+	ULidarPointCloud* OldPointCloud = Cast<ULidarPointCloud>(StaticFindObjectFast(nullptr, InParent, InName, true));
+	if (OldPointCloud)
+	{
+		bOldPointCloudExists = true;
+
+		// If so, store event references to re-apply to the new object
+		OnPointCloudRebuiltEvent = OldPointCloud->OnPointCloudRebuiltEvent;
+		OnPointCloudUpdateCollisionEvent = OldPointCloud->OnPointCloudUpdateCollisionEvent;
+	}
+#endif
+
+	ULidarPointCloud* PointCloud = NewObject<ULidarPointCloud>(InParent, InName, Flags);
+
+#if WITH_EDITOR
+	if (bOldPointCloudExists)
+	{
+		PointCloud->OnPointCloudRebuiltEvent = OnPointCloudRebuiltEvent;
+		PointCloud->OnPointCloudUpdateCollisionEvent = OnPointCloudUpdateCollisionEvent;
+	}
+#endif
+
+	PointCloud->SetSourcePath(Filename);
+	PointCloud->ImportSettings = ImportSettings;
+	PointCloud->Reimport(AsyncParameters);
+
+	return PointCloud;
+}
+
+template<typename T>
+ULidarPointCloud* ULidarPointCloud::CreateFromData(T Points, const int64& Count, const FLidarPointCloudAsyncParameters& AsyncParameters)
+{
+	ULidarPointCloud* PC = NewObject<ULidarPointCloud>();
+
+	// Process points, if there are any available
+	if (Points && Count > 0)
+	{
+		// Start the process
+		if (AsyncParameters.bUseAsync)
+		{
+			Async(EAsyncExecution::Thread, [PC, AsyncParameters, Points, Count]
+			{
+				bool bSuccess = PC->SetData(Points, Count, AsyncParameters.ProgressCallback);				
+				if (AsyncParameters.CompletionCallback)
+				{
+					AsyncParameters.CompletionCallback(bSuccess);
+				}
+			});
+		}
+		else
+		{
+			PC->SetData(Points, Count);
+		}
+	}
+
+	return PC;
+}
+
+FBox ULidarPointCloud::CalculateBoundsFromPoints(const FLidarPointCloudPoint* Points, const int64& Count)
+{
+	FBox Bounds(EForceInit::ForceInit);
+
+	// Process points, if there are any available
+	if (Points && Count > 0)
+	{
+		for (const FLidarPointCloudPoint* Data = Points, *DataEnd = Data + Count; Data != DataEnd; ++Data)
+		{
+			Bounds += Data->Location;
+		}
+	}
+
+	return Bounds;
+}
+
+FBox ULidarPointCloud::CalculateBoundsFromPoints(FLidarPointCloudPoint** Points, const int64& Count)
+{
+	FBox Bounds(EForceInit::ForceInit);
+
+	// Process points, if there are any available
+	if (Points && Count > 0)
+	{
+		for (FLidarPointCloudPoint** Data = Points, **DataEnd = Data + Count; Data != DataEnd; ++Data)
+		{
+			Bounds += (*Data)->Location;
+		}
+	}
+
+	return Bounds;
+}
+
+void ULidarPointCloud::FinishPhysicsAsyncCook(bool bSuccess, UBodySetup* NewBodySetup)
+{
+	FBenchmarkTimer::Log("CookingCollision");
+	Notification.Close(bSuccess);
+
+	if (bSuccess)
+	{
+		BodySetup = NewBodySetup;
+		OnPointCloudUpdateCollisionEvent.Broadcast();
+		InitializeCollisionRendering();
+	}
+
+	bCollisionBuildInProgress = false;
+}
+
+class FPointCloudLatentAction : public FPendingLatentAction
+{
+public:
+	FName ExecutionFunction;
+	int32 OutputLink;
+	FWeakObjectPtr CallbackTarget;
+	ELidarPointCloudAsyncMode* Mode = nullptr;
+
+	FPointCloudLatentAction(const FLatentActionInfo& LatentInfo, ELidarPointCloudAsyncMode& Mode)
+		: ExecutionFunction(LatentInfo.ExecutionFunction)
+		, OutputLink(LatentInfo.Linkage)
+		, CallbackTarget(LatentInfo.CallbackTarget)
+		, Mode(&Mode)
+	{
+	}
+
+	virtual void UpdateOperation(FLatentResponse& Response) override
+	{
+		if (*Mode != ELidarPointCloudAsyncMode::Progress)
+		{
+			Response.FinishAndTriggerIf(true, ExecutionFunction, OutputLink, CallbackTarget);
+		}
+		else
+		{
+			Response.TriggerLink(ExecutionFunction, OutputLink, CallbackTarget);
+		}
+	}
+};
+
+void ULidarPointCloudBlueprintLibrary::CreatePointCloudFromFile(UObject* WorldContextObject, const FString& Filename, bool bUseAsync, FLatentActionInfo LatentInfo, ELidarPointCloudAsyncMode& AsyncMode, float& Progress, ULidarPointCloud*& PointCloud)
+{	
+	CreatePointCloudFromFile(WorldContextObject, Filename, bUseAsync, LatentInfo, FLidarPointCloudImportSettings::MakeGeneric(Filename), AsyncMode, Progress, PointCloud);
+}
+
+void ULidarPointCloudBlueprintLibrary::CreatePointCloudFromFile(UObject* WorldContextObject, const FString& Filename, bool bUseAsync, FLatentActionInfo LatentInfo, TSharedPtr<FLidarPointCloudImportSettings> ImportSettings, ELidarPointCloudAsyncMode& AsyncMode, float& Progress, ULidarPointCloud*& PointCloud)
+{
+	PointCloud = nullptr;
+	if (UWorld* World = GEngine->GetWorldFromContextObject(WorldContextObject, EGetWorldErrorMode::LogAndReturnNull))
+	{
+		FLatentActionManager& LatentActionManager = World->GetLatentActionManager();
+		if (LatentActionManager.FindExistingAction<FPointCloudLatentAction>(LatentInfo.CallbackTarget, LatentInfo.UUID) == nullptr)
+		{
+			AsyncMode = ELidarPointCloudAsyncMode::Progress;
+			FPointCloudLatentAction* CompletionAction = new FPointCloudLatentAction(LatentInfo, AsyncMode);
+
+			LatentActionManager.AddNewAction(LatentInfo.CallbackTarget, LatentInfo.UUID, CompletionAction);
+
+			PointCloud = ULidarPointCloud::CreateFromFile(Filename, FLidarPointCloudAsyncParameters(bUseAsync,
+				[&Progress, &AsyncMode](float InProgress)
+				{
+					Progress = InProgress;
+				},
+				[&AsyncMode](bool bSuccess)
+				{
+					AsyncMode = bSuccess ? ELidarPointCloudAsyncMode::Success : ELidarPointCloudAsyncMode::Failure;
+				}),
+				ImportSettings);
+		}
+	}
+}
+
+void ULidarPointCloudBlueprintLibrary::CreatePointCloudFromData(UObject* WorldContextObject, const TArray<FLidarPointCloudPoint>& Points, bool bUseAsync, FLatentActionInfo LatentInfo, ELidarPointCloudAsyncMode& AsyncMode, float& Progress, ULidarPointCloud*& PointCloud)
+{
+	PointCloud = nullptr;
+	if (UWorld* World = GEngine->GetWorldFromContextObject(WorldContextObject, EGetWorldErrorMode::LogAndReturnNull))
+	{
+		FLatentActionManager& LatentActionManager = World->GetLatentActionManager();
+		if (LatentActionManager.FindExistingAction<FPointCloudLatentAction>(LatentInfo.CallbackTarget, LatentInfo.UUID) == nullptr)
+		{
+			AsyncMode = ELidarPointCloudAsyncMode::Progress;
+			FPointCloudLatentAction* CompletionAction = new FPointCloudLatentAction(LatentInfo, AsyncMode);
+
+			LatentActionManager.AddNewAction(LatentInfo.CallbackTarget, LatentInfo.UUID, CompletionAction);
+
+			PointCloud = ULidarPointCloud::CreateFromData(Points, FLidarPointCloudAsyncParameters(bUseAsync,
+				[&Progress, &AsyncMode](float InProgress)
+				{
+					Progress = InProgress;
+				},
+				[&AsyncMode](bool bSuccess)
+				{
+					AsyncMode = bSuccess ? ELidarPointCloudAsyncMode::Success : ELidarPointCloudAsyncMode::Failure;
+				}));
+		}
+	}
+}
+
+bool ULidarPointCloudBlueprintLibrary::LineTraceSingle(UObject* WorldContextObject, FVector Origin, FVector Direction, float Radius, bool bVisibleOnly, FLidarPointCloudTraceHit& Hit)
+{
+	if (UWorld* World = WorldContextObject ? WorldContextObject->GetWorld() : nullptr)
+	{
+		const FLidarPointCloudRay Ray(Origin, Direction);
+
+		for (TActorIterator<ALidarPointCloudActor> Itr(World); Itr; ++Itr)
+		{
+			auto Actor = *Itr;
+			auto Component = Actor->GetPointCloudComponent();
+			auto Point = Component->LineTraceSingle(Ray, Radius, bVisibleOnly);
+			if (Point)
+			{
+				Hit = FLidarPointCloudTraceHit(Actor, Component);
+				Hit.Points.Add(*Point);
+				return true;
+			}
+		}
+	}
+
+	return false;
+}
+
+bool ULidarPointCloudBlueprintLibrary::LineTraceMulti(UObject* WorldContextObject, FVector Origin, FVector Direction, float Radius, bool bVisibleOnly, TArray<FLidarPointCloudTraceHit>& Hits)
+{
+	Hits.Reset();
+
+	if (UWorld * World = WorldContextObject ? WorldContextObject->GetWorld() : nullptr)
+	{
+		const FLidarPointCloudRay Ray(Origin, Direction);
+
+		for (TActorIterator<ALidarPointCloudActor> Itr(World); Itr; ++Itr)
+		{
+			auto Actor = *Itr;
+			auto Component = Actor->GetPointCloudComponent();
+
+			FLidarPointCloudTraceHit Hit(Actor, Component);
+			if (Component->LineTraceMulti(Ray, Radius, bVisibleOnly, Hit.Points))
+			{
+				Hits.Add(Hit);
+				return true;
+			}
+		}
+	}
+
+	return Hits.Num() > 0;
+}
+
+template bool ULidarPointCloud::InsertPoints_NoLock<const FLidarPointCloudPoint*>(const FLidarPointCloudPoint*, const int64&, ELidarPointCloudDuplicateHandling, bool, FThreadSafeBool*, TFunction<void(float)>);
+template bool ULidarPointCloud::InsertPoints_NoLock<FLidarPointCloudPoint**>(FLidarPointCloudPoint**, const int64&, ELidarPointCloudDuplicateHandling, bool, FThreadSafeBool*, TFunction<void(float)>);
+template bool ULidarPointCloud::SetData<const FLidarPointCloudPoint*>(const FLidarPointCloudPoint*, const int64&, TFunction<void(float)>);
+template bool ULidarPointCloud::SetData<FLidarPointCloudPoint**>(FLidarPointCloudPoint**, const int64&, TFunction<void(float)>);
+
+#undef LOCTEXT_NAMESPACE
\ No newline at end of file
