Copy up Dev-Enterprise @ CL11133089

diff --git a/Engine/Plugins/Enterprise/LidarPointCloud/Source/LidarPointCloudRuntime/Private/LidarPointCloudOctree.cpp b/Engine/Plugins/Enterprise/LidarPointCloud/Source/LidarPointCloudRuntime/Private/LidarPointCloudOctree.cpp
new file mode 100644
--- /dev/null
+++ b/Engine/Plugins/Enterprise/LidarPointCloud/Source/LidarPointCloudRuntime/Private/LidarPointCloudOctree.cpp
@@ -0,0 +1,1949 @@
+// Copyright Epic Games, Inc. All Rights Reserved.
+ 
+#include "LidarPointCloudOctree.h"
+#include "LidarPointCloudOctreeMacros.h"
+#include "LidarPointCloud.h"
+#include "Collision/LidarPointCloudCollision.h"
+#include "Misc/ScopeTryLock.h"
+#include "Containers/Queue.h"
+#include "SceneView.h"
+#include "Async/Async.h"
+#include "Misc/FileHelper.h"
+
+DECLARE_CYCLE_STAT(TEXT("Node Selection"), STAT_NodeSelection, STATGROUP_LidarPointCloud);
+
+int32 FLidarPointCloudOctree::MaxNodeDepth = (1 << (sizeof(FLidarPointCloudOctreeNode::Depth) * 8)) - 1;
+int32 FLidarPointCloudOctree::MaxBucketSize = 200;
+int32 FLidarPointCloudOctree::NodeGridResolution = 96;
+
+/** Used for grid allocation calculations */
+struct FGridAllocation
+{
+	/** Index to the point inside of the AllocatedPoints */
+	int32 Index;
+
+	/** Index of the child node this point would be placed in */
+	int32 ChildNodeLocation;
+
+	/** The calculated distance squared from the center of the given point */
+	float DistanceFromCenter;
+
+	FGridAllocation() {}
+	FGridAllocation(const int32& Index, const FGridAllocation& GridAllocation)
+		: Index(Index)
+		, ChildNodeLocation(GridAllocation.ChildNodeLocation)
+		, DistanceFromCenter(GridAllocation.DistanceFromCenter)
+	{
+	}
+};
+
+/** Used for node size sorting and node selection. */
+struct FNodeSizeData
+{
+	FLidarPointCloudTraversalOctreeNode* Node;
+	float Size;
+	FNodeSizeData(FLidarPointCloudTraversalOctreeNode* Node, const float& Size) : Node(Node), Size(Size) {}
+};
+
+auto PointVisibilitySortingFunction = [](const FLidarPointCloudPoint& A, const FLidarPointCloudPoint& B) { return A.bVisible > B.bVisible; };
+
+FGridAllocation CalculateGridCellData(const FLidarPointCloudPoint* Point, const FVector& Center, const FLidarPointCloudOctree::FSharedLODData& LODData)
+{
+	FVector CenterRelativeLocation = Point->Location - Center;
+	FVector OffsetLocation = CenterRelativeLocation + LODData.Extent;
+	FVector NormalizedGridLocation = OffsetLocation * LODData.NormalizationMultiplier;
+
+	// Calculate the location on this node's Grid
+	int32 GridX = FMath::Min(FLidarPointCloudOctree::NodeGridResolution - 1, (int32)NormalizedGridLocation.X);
+	int32 GridY = FMath::Min(FLidarPointCloudOctree::NodeGridResolution - 1, (int32)NormalizedGridLocation.Y);
+	int32 GridZ = FMath::Min(FLidarPointCloudOctree::NodeGridResolution - 1, (int32)NormalizedGridLocation.Z);
+
+	FGridAllocation Allocation;
+	Allocation.Index = GridX * FLidarPointCloudOctree::NodeGridResolution * FLidarPointCloudOctree::NodeGridResolution + GridY * FLidarPointCloudOctree::NodeGridResolution + GridZ;
+	Allocation.DistanceFromCenter = (FVector(GridX + 0.5f, GridY + 0.5f, GridZ + 0.5f) * LODData.GridSize - OffsetLocation).SizeSquared();
+	Allocation.ChildNodeLocation = (CenterRelativeLocation.X > 0 ? 4 : 0) + (CenterRelativeLocation.Y > 0 ? 2 : 0) + (CenterRelativeLocation.Z > 0);
+
+	return Allocation;
+}
+
+FORCEINLINE float BrightnessFromColor(const FColor& Color) { return 0.2126 * Color.R + 0.7152 * Color.G + 0.0722 * Color.B; }
+
+bool IsOnBoundsEdge(const FBox& Bounds, const FVector& Location)
+{
+	return (Location.X == Bounds.Min.X) || (Location.X == Bounds.Max.X) || (Location.Y == Bounds.Min.Y) || (Location.Y == Bounds.Max.Y) || (Location.Z == Bounds.Min.Z) || (Location.Z == Bounds.Max.Z);
+}
+
+//////////////////////////////////////////////////////////// FSharedLODData
+
+FLidarPointCloudOctree::FSharedLODData::FSharedLODData(const FVector& InExtent)
+	: Radius(InExtent.Size())
+	, RadiusSq(FMath::Square(InExtent.Size()))
+	, Extent(InExtent)
+{
+	GridSize = Extent * 2 / FLidarPointCloudOctree::NodeGridResolution;
+	NormalizationMultiplier = FVector(FLidarPointCloudOctree::NodeGridResolution) / (Extent * 2);
+}
+
+//////////////////////////////////////////////////////////// FLidarPointCloudOctreeNode
+
+FLidarPointCloudOctreeNode::FLidarPointCloudOctreeNode(FLidarPointCloudOctree* Tree, FLidarPointCloudOctreeNode* Parent, const uint8& Depth, const uint8& LocationInParent, const FVector& Center)
+	: Depth(Depth)
+	, LocationInParent(LocationInParent)
+	, Center(Center)
+	, bVisibilityDirty(false)
+	, NumVisiblePoints(0)
+{
+	if (Tree)
+	{
+		Tree->NodeCount[Depth].Increment();
+	}
+}
+
+FLidarPointCloudOctreeNode::~FLidarPointCloudOctreeNode()
+{
+	for (int32 i = 0; i < Children.Num(); i++)
+	{
+		delete Children[i];
+		Children[i] = nullptr;
+	}
+}
+
+void FLidarPointCloudOctreeNode::UpdateNumVisiblePoints()
+{
+	if (bVisibilityDirty)
+	{
+		// Sort points to speed up rendering
+		Algo::Sort(AllocatedPoints, MoveTemp(PointVisibilitySortingFunction));
+		Algo::Sort(PaddingPoints, MoveTemp(PointVisibilitySortingFunction));
+
+		// Recalculate visibility
+		NumVisiblePoints = 0;
+		FOR(Point, AllocatedPoints)
+		{
+			if (!Point->bVisible)
+			{
+				break;
+			}
+
+			NumVisiblePoints++;
+		}
+
+		FOR(Point, PaddingPoints)
+		{
+			if (!Point->bVisible)
+			{
+				break;
+			}
+
+			NumVisiblePoints++;
+		}
+
+		bVisibilityDirty = false;
+	}
+}
+
+FBox FLidarPointCloudOctreeNode::GetBounds(const FLidarPointCloudOctree* Tree) const
+{
+	return FBox(Center - Tree->SharedData[Depth].Extent, Center + Tree->SharedData[Depth].Extent);
+}
+
+FSphere FLidarPointCloudOctreeNode::GetSphereBounds(const FLidarPointCloudOctree* Tree) const
+{
+	return FSphere(Center, Tree->SharedData[Depth].Radius);
+}
+
+FLidarPointCloudOctreeNode* FLidarPointCloudOctreeNode::GetChildNodeAtLocation(const uint8& Location) const
+{
+	for (FLidarPointCloudOctreeNode* Child : Children)
+	{
+		if (Child->LocationInParent == Location)
+		{
+			return Child;
+		}
+	}
+
+	return nullptr;
+}
+
+void FLidarPointCloudOctreeNode::InsertPoints(FLidarPointCloudOctree* Tree, const FLidarPointCloudPoint* Points, const int32& Count, ELidarPointCloudDuplicateHandling DuplicateHandling)
+{
+	const auto& LODData = Tree->SharedData[Depth];
+
+	// Local 
+	TArray<FLidarPointCloudPoint> PointBuckets[8];
+	TMultiMap<int32, FGridAllocation> NewGridAllocationMap, CurrentGridAllocationMap;
+
+	int32 NumPointsAdded = 0;
+
+	const float MaxDistanceForDuplicate = GetDefault<ULidarPointCloudSettings>()->MaxDistanceForDuplicate;
+
+	// Filter the local set of incoming data
+	for (int32 Index = 0; Index < Count; ++Index)
+	{
+		auto InGridData = CalculateGridCellData(&Points[Index], Center, LODData);
+		FGridAllocation* GridCell = NewGridAllocationMap.Find(InGridData.Index);
+
+		// Attempt to allocate the point to this node
+		if (GridCell)
+		{
+			bool bStoreInBucket = true;
+
+			if (DuplicateHandling != ELidarPointCloudDuplicateHandling::Ignore && Points[GridCell->Index].Location.Equals(Points[Index].Location, MaxDistanceForDuplicate))
+			{
+				if (DuplicateHandling == ELidarPointCloudDuplicateHandling::SelectFirst || BrightnessFromColor(Points[Index].Color) <= BrightnessFromColor(Points[GridCell->Index].Color))
+				{
+					continue;
+				}
+				else 
+				{
+					bStoreInBucket = false;
+				}				
+			}
+
+			if (InGridData.DistanceFromCenter < GridCell->DistanceFromCenter)
+			{
+				if (bStoreInBucket)
+				{
+					PointBuckets[GridCell->ChildNodeLocation].Add(Points[GridCell->Index]);
+				}
+				
+				GridCell->Index = Index;
+				GridCell->DistanceFromCenter = InGridData.DistanceFromCenter;
+			}
+			else if(bStoreInBucket)
+			{
+				PointBuckets[InGridData.ChildNodeLocation].Add(Points[Index]);
+			}
+		}
+		else
+		{
+			NewGridAllocationMap.Add(InGridData.Index, FGridAllocation(Index, InGridData));
+		}
+	}
+
+	// Process incoming points
+	{
+		FScopeLock Lock(&MapLock);
+
+		// Rebuild Current Grid Mapping
+		for (int32 i = 0; i < AllocatedPoints.Num(); ++i)
+		{
+			auto InGridData = CalculateGridCellData(&AllocatedPoints[i], Center, LODData);
+			CurrentGridAllocationMap.Add(InGridData.Index, FGridAllocation(i, InGridData));
+		}
+
+		// Compare the incoming data to the currently held set, and replace if necessary
+		for (auto& Element : NewGridAllocationMap)
+		{
+			const int32& GridIndex = Element.Key;
+			const FLidarPointCloudPoint& Point = Points[Element.Value.Index];
+			FGridAllocation* GridCell = CurrentGridAllocationMap.Find(GridIndex);
+
+			// Attempt to allocate the point to this node
+			if (GridCell)
+			{
+				bool bStoreInBucket = true;
+
+				if (DuplicateHandling != ELidarPointCloudDuplicateHandling::Ignore && AllocatedPoints[GridCell->Index].Location.Equals(Point.Location, MaxDistanceForDuplicate))
+				{
+					if (DuplicateHandling == ELidarPointCloudDuplicateHandling::SelectFirst || BrightnessFromColor(Point.Color) <= BrightnessFromColor(AllocatedPoints[GridCell->Index].Color))
+					{
+						continue;
+					}
+					else
+					{
+						bStoreInBucket = false;
+					}
+				}
+
+				// If the new point's distance from center of node is shorter than the existing point's, replace the point
+				if (Element.Value.DistanceFromCenter < GridCell->DistanceFromCenter)
+				{
+					if (bStoreInBucket)
+					{
+						PointBuckets[GridCell->ChildNodeLocation].Add(AllocatedPoints[GridCell->Index]);
+					}
+
+					AllocatedPoints[GridCell->Index] = Point;
+					GridCell->DistanceFromCenter = Element.Value.DistanceFromCenter;
+				}
+				// ... otherwise add it straight to the bucket
+				else if (bStoreInBucket)
+				{
+					PointBuckets[Element.Value.ChildNodeLocation].Add(Point);
+				}
+			}
+			else
+			{
+				CurrentGridAllocationMap.Add(GridIndex, FGridAllocation(AllocatedPoints.Add(Point), Element.Value));
+				NumPointsAdded++;
+			}
+		}
+
+		// Distribute PaddingPoints into the relevant arrays of PointBuckets
+		for (auto& Point : PaddingPoints)
+		{
+			PointBuckets[CalculateGridCellData(&Point, Center, LODData).ChildNodeLocation].Add(Point);
+		}
+		NumPointsAdded -= PaddingPoints.Num();
+		PaddingPoints.Reset();
+
+		for (uint8 i = 0; i < 8; ++i)
+		{
+			if (!GetChildNodeAtLocation(i))
+			{
+				// While the threads are locked, check if any child nodes need creating
+				if (Depth < FLidarPointCloudOctree::MaxNodeDepth && PointBuckets[i].Num() > FLidarPointCloudOctree::MaxBucketSize)
+				{
+					const FVector ChildNodeCenter = Center + LODData.Extent * (FVector(-0.5f) + FVector((i & 4) == 4, (i & 2) == 2, (i & 1) == 1));
+					Children.Add(new FLidarPointCloudOctreeNode(Tree, this, Depth + 1, i, ChildNodeCenter));
+
+					// The recursive InserPoints call will happen later, after the Lock is released
+				}
+				// ... otherwise, points can be re-added back as padding
+				else
+				{
+					PaddingPoints.Append(PointBuckets[i]);
+					PointBuckets[i].Reset();
+				}
+			}
+		}
+
+		NumPointsAdded += PaddingPoints.Num();
+
+		// Shrink the data usage
+		AllocatedPoints.Shrink();
+		PaddingPoints.Shrink();
+	}
+	
+	AddPointCount(Tree, NumPointsAdded);
+
+	// Pass surplus points
+	for (uint8 i = 0; i < 8; ++i)
+	{
+		if (PointBuckets[i].Num() > 0)
+		{
+			GetChildNodeAtLocation(i)->InsertPoints(Tree, PointBuckets[i].GetData(), PointBuckets[i].Num(), DuplicateHandling);
+		}
+	}
+}
+
+void FLidarPointCloudOctreeNode::InsertPoints(FLidarPointCloudOctree* Tree, FLidarPointCloudPoint** Points, const int32& Count, ELidarPointCloudDuplicateHandling DuplicateHandling)
+{
+	const auto& LODData = Tree->SharedData[Depth];
+
+	// Local 
+	TArray<FLidarPointCloudPoint> PointBuckets[8];
+	TMultiMap<int32, FGridAllocation> NewGridAllocationMap, CurrentGridAllocationMap;
+
+	int32 NumPointsAdded = 0;
+
+	const float MaxDistanceForDuplicate = GetDefault<ULidarPointCloudSettings>()->MaxDistanceForDuplicate;
+
+	// Filter the local set of incoming data
+	for (int32 Index = 0; Index < Count; Index++)
+	{
+		auto InGridData = CalculateGridCellData(Points[Index], Center, LODData);
+		FGridAllocation* GridCell = NewGridAllocationMap.Find(InGridData.Index);
+
+		// Attempt to allocate the point to this node
+		if (GridCell)
+		{
+			bool bStoreInBucket = true;
+
+			if (DuplicateHandling != ELidarPointCloudDuplicateHandling::Ignore && Points[GridCell->Index]->Location.Equals(Points[Index]->Location, MaxDistanceForDuplicate))
+			{
+				if (DuplicateHandling == ELidarPointCloudDuplicateHandling::SelectFirst || BrightnessFromColor(Points[Index]->Color) <= BrightnessFromColor(Points[GridCell->Index]->Color))
+				{
+					continue;
+				}
+				else
+				{
+					bStoreInBucket = false;
+				}
+			}
+
+			if (InGridData.DistanceFromCenter < GridCell->DistanceFromCenter)
+			{
+				if (bStoreInBucket)
+				{
+					PointBuckets[GridCell->ChildNodeLocation].Emplace(*Points[GridCell->Index]);
+				}
+
+				GridCell->Index = Index;
+				GridCell->DistanceFromCenter = InGridData.DistanceFromCenter;
+			}
+			else if (bStoreInBucket)
+			{
+				PointBuckets[InGridData.ChildNodeLocation].Emplace(*Points[Index]);
+			}
+		}
+		else
+		{
+			NewGridAllocationMap.Add(InGridData.Index, FGridAllocation(Index, InGridData));
+		}
+	}
+
+	// Process incoming points
+	{
+		FScopeLock Lock(&MapLock);
+
+		// Rebuild Current Grid Mapping
+		for (int32 i = 0; i < AllocatedPoints.Num(); i++)
+		{
+			auto InGridData = CalculateGridCellData(&AllocatedPoints[i], Center, LODData);
+			CurrentGridAllocationMap.Add(InGridData.Index, FGridAllocation(i, InGridData));
+		}
+
+		// Compare the incoming data to the currently held set, and replace if necessary
+		FLidarPointCloudPoint Point;
+		for (auto& Element : NewGridAllocationMap)
+		{
+			const int32& GridIndex = Element.Key;
+			Point.CopyFrom(*Points[Element.Value.Index]);
+			FGridAllocation* GridCell = CurrentGridAllocationMap.Find(GridIndex);
+
+			// Attempt to allocate the point to this node
+			if (GridCell)
+			{
+				bool bStoreInBucket = true;
+
+				if (DuplicateHandling != ELidarPointCloudDuplicateHandling::Ignore && AllocatedPoints[GridCell->Index].Location.Equals(Point.Location, MaxDistanceForDuplicate))
+				{
+					if (DuplicateHandling == ELidarPointCloudDuplicateHandling::SelectFirst || BrightnessFromColor(Point.Color) <= BrightnessFromColor(AllocatedPoints[GridCell->Index].Color))
+					{
+						continue;
+					}
+					else
+					{
+						bStoreInBucket = false;
+					}
+				}
+
+				// If the new point's distance from center of node is shorter than the existing point's, replace the point
+				if (Element.Value.DistanceFromCenter < GridCell->DistanceFromCenter)
+				{
+					if (bStoreInBucket)
+					{
+						PointBuckets[GridCell->ChildNodeLocation].Add(AllocatedPoints[GridCell->Index]);
+					}
+
+					AllocatedPoints[GridCell->Index] = Point;
+					GridCell->DistanceFromCenter = Element.Value.DistanceFromCenter;
+				}
+				// ... otherwise add it straight to the bucket
+				else if (bStoreInBucket)
+				{
+					PointBuckets[Element.Value.ChildNodeLocation].Add(Point);
+				}
+			}
+			else
+			{
+				CurrentGridAllocationMap.Add(GridIndex, FGridAllocation(AllocatedPoints.Add(Point), Element.Value));
+				NumPointsAdded++;
+			}
+		}
+
+		// Distribute PaddingPoints into the relevant arrays of PointBuckets
+		for (auto& PaddingPoint : PaddingPoints)
+		{
+			PointBuckets[CalculateGridCellData(&PaddingPoint, Center, LODData).ChildNodeLocation].Add(PaddingPoint);
+		}
+		NumPointsAdded -= PaddingPoints.Num();
+		PaddingPoints.Reset();
+
+		for (uint8 i = 0; i < 8; i++)
+		{
+			if (!GetChildNodeAtLocation(i))
+			{
+				// While the threads are locked, check if any child nodes need creating
+				if (Depth < FLidarPointCloudOctree::MaxNodeDepth && PointBuckets[i].Num() > FLidarPointCloudOctree::MaxBucketSize)
+				{
+					const FVector ChildNodeCenter = Center + LODData.Extent * (FVector(-0.5f) + FVector((i & 4) == 4, (i & 2) == 2, (i & 1) == 1));
+					Children.Add(new FLidarPointCloudOctreeNode(Tree, this, Depth + 1, i, ChildNodeCenter));
+
+					// The recursive InserPoints call will happen later, after the Lock is released
+				}
+				// ... otherwise, points can be re-added back as padding
+				else
+				{
+					PaddingPoints.Append(PointBuckets[i]);
+					PointBuckets[i].Reset();
+				}
+			}
+		}
+
+		NumPointsAdded += PaddingPoints.Num();
+
+		// Shrink the data usage
+		AllocatedPoints.Shrink();
+		PaddingPoints.Shrink();
+	}
+
+	AddPointCount(Tree, NumPointsAdded);
+
+	// Pass surplus points
+	for (uint8 i = 0; i < 8; i++)
+	{
+		if (PointBuckets[i].Num() > 0)
+		{
+			GetChildNodeAtLocation(i)->InsertPoints(Tree, PointBuckets[i].GetData(), PointBuckets[i].Num(), DuplicateHandling);
+		}
+	}
+}
+
+void FLidarPointCloudOctreeNode::Empty(bool bRecursive)
+{
+	AllocatedPoints.Empty();
+	PaddingPoints.Empty();
+
+	if (bRecursive)
+	{
+		for (int32 i = 0; i < Children.Num(); i++)
+		{
+			Children[i]->Empty(true);
+		}
+	}
+}
+
+uint32 FLidarPointCloudOctreeNode::GetMaxDepth() const
+{
+	uint32 MaxDepth = Depth;
+
+	for (FLidarPointCloudOctreeNode* Child : Children)
+	{
+		MaxDepth = FMath::Max(MaxDepth, Child->GetMaxDepth());
+	}
+
+	return MaxDepth;
+}
+
+int64 FLidarPointCloudOctreeNode::GetAllocatedSize(bool bRecursive) const
+{
+	int64 Size = sizeof(FLidarPointCloudOctreeNode);
+
+	Size += Children.GetAllocatedSize();
+	Size += AllocatedPoints.GetAllocatedSize();
+	Size += PaddingPoints.GetAllocatedSize();
+
+	if (bRecursive)
+	{
+		for (FLidarPointCloudOctreeNode* Child : Children)
+		{
+			Size += Child->GetAllocatedSize(true);
+		}
+	}
+
+	return Size;
+}
+
+void FLidarPointCloudOctreeNode::Serialize(FArchive& Ar, FLidarPointCloudOctree* Tree)
+{
+	// This is needed as LocationInParent doesn't use a full byte and the FArchive complains
+	uint8 FullLocationInParent = LocationInParent;
+	Ar << FullLocationInParent << Center << AllocatedPoints << PaddingPoints;
+
+	int32 NumChildren = Children.Num();
+	Ar << NumChildren;
+
+	if (Ar.IsLoading())
+	{
+		LocationInParent = FullLocationInParent;
+
+		AddPointCount(Tree);
+
+		bVisibilityDirty = true;
+
+		Children.AddUninitialized(NumChildren);
+		for (int32 i = 0; i < NumChildren; i++)
+		{
+			Children[i] = new FLidarPointCloudOctreeNode(Tree, this, Depth + 1);
+		}
+	}
+
+	for (int32 i = 0; i < NumChildren; i++)
+	{
+		Children[i]->Serialize(Ar, Tree);
+	}
+}
+
+void FLidarPointCloudOctreeNode::AddPointCount(FLidarPointCloudOctree* Tree, int32 PointCount /*= INT32_MIN*/)
+{
+	const int64 Count = PointCount == INT32_MIN ? GetNumPoints() : PointCount;
+
+	Tree->PointCount[Depth].Add(Count);
+	NumVisiblePoints += Count;
+}
+
+//////////////////////////////////////////////////////////// FLidarPointCloudOctree
+
+FLidarPointCloudOctree::FLidarPointCloudOctree()
+{
+	PointCount.AddDefaulted(MaxNodeDepth + 1);
+	NodeCount.AddDefaulted(MaxNodeDepth + 1);
+	SharedData.AddDefaulted(MaxNodeDepth + 1);
+}
+
+FLidarPointCloudOctree::~FLidarPointCloudOctree()
+{
+	MarkTraversalOctreesForInvalidation();
+}
+
+int32 FLidarPointCloudOctree::GetNumLODs() const
+{
+	int32 LODCount = 0;
+
+	for (; LODCount < NodeCount.Num(); LODCount++)
+	{
+		if (NodeCount[LODCount].GetValue() == 0)
+		{
+			break;
+		}
+	}
+
+	return LODCount;
+}
+
+void FLidarPointCloudOctree::RefreshPointsBounds()
+{
+	PointsBounds.Init();
+
+	TQueue<const FLidarPointCloudOctreeNode*> Nodes;
+	const FLidarPointCloudOctreeNode* CurrentNode = nullptr;
+	Nodes.Enqueue(&Root);
+	while (Nodes.Dequeue(CurrentNode))
+	{
+		FOR(Point, CurrentNode->AllocatedPoints)
+		{
+			PointsBounds += Point->Location;
+		}
+
+		FOR(Point, CurrentNode->PaddingPoints)
+		{
+			PointsBounds += Point->Location;
+		}
+
+		for (auto Child : CurrentNode->Children)
+		{
+			Nodes.Enqueue(Child);
+		}
+	}
+}
+
+int64 FLidarPointCloudOctree::GetNumPoints() const
+{
+	int64 TotalPointCount = 0;
+
+	for (int32 i = 0; i < PointCount.Num(); i++)
+	{
+		int64 NumPoints = PointCount[i].GetValue();
+
+		if (NumPoints > 0)
+		{
+			TotalPointCount += NumPoints;
+		}
+		else
+		{
+			break;
+		}
+	}
+
+	return TotalPointCount;
+}
+
+int32 FLidarPointCloudOctree::GetNumNodes() const
+{
+	int32 TotalNodeCount = 0;
+
+	for (int32 i = 0; i < NodeCount.Num(); i++)
+	{
+		int32 NumNodes = NodeCount[i].GetValue();
+		if (NumNodes > 0)
+		{
+			TotalNodeCount += NumNodes;
+		}
+		else
+		{
+			break;
+		}
+	}
+
+	return TotalNodeCount;
+}
+
+int64 FLidarPointCloudOctree::GetAllocatedSize() const
+{
+	// Use cached version, if possible
+	if (PreviousPointCount == GetNumPoints() && PreviousNodeCount == GetNumNodes())
+	{
+		return PreviousAllocatedSize;
+	}
+
+	FScopeTryLock Lock(&DataLock);
+	if (!Lock.IsLocked())
+	{
+		return 0;
+	}
+
+	PreviousPointCount = GetNumPoints();
+	PreviousNodeCount = GetNumNodes();
+
+	PreviousAllocatedSize = sizeof(FLidarPointCloudOctree);
+
+	PreviousAllocatedSize += SharedData.GetAllocatedSize();
+	PreviousAllocatedSize += PointCount.GetAllocatedSize();
+
+	PreviousAllocatedSize += Root.GetAllocatedSize(true);
+
+	return PreviousAllocatedSize;
+}
+
+float FLidarPointCloudOctree::GetEstimatedPointSpacing() const
+{
+	float Spacing = 0;
+	const int64 TotalPointCount = GetNumPoints();
+
+	for (int32 i = 0; i < PointCount.Num(); ++i)
+	{
+		Spacing += SharedData[i].GridSize.GetMax() * PointCount[i].GetValue() / TotalPointCount;
+	}
+
+	return Spacing;
+}
+
+void FLidarPointCloudOctree::BuildCollision(const float& Accuracy, const bool& bVisibleOnly)
+{
+	LidarPointCloudCollision::BuildCollisionMesh(this, Accuracy, bVisibleOnly, &CollisionMesh);
+}
+
+void FLidarPointCloudOctree::RemoveCollision()
+{
+	FScopeLock Lock(&DataLock);
+
+	CollisionMesh.~FTriMeshCollisionData();
+	new (&CollisionMesh) FTriMeshCollisionData();
+}
+
+void FLidarPointCloudOctree::GetPoints(TArray<FLidarPointCloudPoint*>& Points, int64 StartIndex, int64 Count) const
+{
+	check(StartIndex >= 0 && StartIndex < GetNumPoints());
+
+	if (Count < 0)
+	{
+		Count = GetNumPoints();
+	}
+
+	Count = FMath::Min(Count, GetNumPoints() - StartIndex);
+
+	// TArray only supports int32 ax max number of elements!
+	check(Count <= INT32_MAX);
+
+	Points.Empty(Count);
+
+	FScopeLock Lock(&DataLock);
+
+	ITERATE_NODES_CONST(
+	{
+		// If no data is required, quit
+		if (Count == 0)
+		{
+			return;
+		}
+
+		// Should this node's points be injected?
+		if (StartIndex < CurrentNode->GetNumPoints())
+		{
+			// If the index is at 0 and the requested count is at least the size of this node, append whole arrays
+			if (StartIndex == 0 && Count >= CurrentNode->GetNumPoints())
+			{
+				// Expand the array
+				int32 Offset = Points.Num();
+				Points.AddUninitialized(CurrentNode->GetNumPoints());
+				FLidarPointCloudPoint** Dest = Points.GetData() + Offset;
+
+				// Assign pointers to the data
+				for (auto& Point : CurrentNode->AllocatedPoints)
+				{
+					*Dest++ = &Point;
+				}
+				for (auto& Point : CurrentNode->PaddingPoints)
+				{
+					*Dest++ = &Point;
+				}
+
+				Count -= CurrentNode->GetNumPoints();
+			}
+			// ... otherwise iterate over needed data
+			else
+			{
+				// AllocatedPoints
+				{
+					int32 NumPointsToCopy = FMath::Max(0LL, FMath::Min((int64)CurrentNode->AllocatedPoints.Num(), Count + StartIndex) - StartIndex);
+					if (NumPointsToCopy > 0)
+					{
+						// Expand the array
+						int32 Offset = Points.Num();
+						Points.AddUninitialized(NumPointsToCopy);
+												
+						// Setup pointers
+						FLidarPointCloudPoint* Src = CurrentNode->AllocatedPoints.GetData() + StartIndex;
+						FLidarPointCloudPoint** Dest = Points.GetData() + Offset;
+
+						// Assign source pointers to the destination
+						for (int32 i = 0; i < NumPointsToCopy; i++)
+						{
+							*Dest++ = Src++;
+						}
+
+						StartIndex = FMath::Max(0LL, StartIndex - NumPointsToCopy);
+						Count -= NumPointsToCopy;
+					}
+				}
+
+				// PaddingPoints
+				{
+					int32 NumPointsToCopy = FMath::Max(0LL, FMath::Min((int64)CurrentNode->PaddingPoints.Num(), Count + StartIndex) - StartIndex);
+					if (NumPointsToCopy > 0)
+					{
+						// Expand the array
+						int32 Offset = Points.Num();
+						Points.AddUninitialized(NumPointsToCopy);
+
+						// Setup pointers
+						FLidarPointCloudPoint* Src = CurrentNode->PaddingPoints.GetData() + StartIndex;
+						FLidarPointCloudPoint** Dest = Points.GetData() + Offset;
+
+						// Assign source pointers to the destination
+						for (int32 i = 0; i < NumPointsToCopy; i++)
+						{
+							*Dest++ = Src++;
+						}
+
+						StartIndex = FMath::Max(0LL, StartIndex - NumPointsToCopy);;
+						Count -= NumPointsToCopy;
+					}
+				}
+			}
+		}
+		// ... or skipped completely?
+		else
+		{
+			StartIndex -= CurrentNode->GetNumPoints();
+		}
+	}, true);
+}
+
+void FLidarPointCloudOctree::GetPointsInSphere(TArray<FLidarPointCloudPoint*>& SelectedPoints, const FSphere& Sphere, const bool& bVisibleOnly) const
+{
+	SelectedPoints.Reset();
+	PROCESS_IN_SPHERE_CONST({ SelectedPoints.Add(Point); });
+}
+
+void FLidarPointCloudOctree::GetPointsInBox(TArray<FLidarPointCloudPoint*>& SelectedPoints, const FBox& Box, const bool& bVisibleOnly) const
+{
+	SelectedPoints.Reset(); 
+	PROCESS_IN_BOX_CONST({ SelectedPoints.Add(Point); });
+}
+
+void FLidarPointCloudOctree::GetPointsInFrustum(TArray<FLidarPointCloudPoint*>& SelectedPoints, const FConvexVolume& Frustum, const bool& bVisibleOnly) const
+{
+	SelectedPoints.Reset();
+	PROCESS_IN_FRUSTUM_CONST({ SelectedPoints.Add(Point); });
+}
+
+void FLidarPointCloudOctree::GetPointsAsCopies(TArray<FLidarPointCloudPoint>& Points, int64 StartIndex, int64 Count) const
+{
+	// If empty, abort
+	if (GetNumPoints() == 0)
+	{
+		return;
+	}
+
+	// Make sure to operate on correct range
+	check(StartIndex >= 0 && StartIndex < GetNumPoints());
+
+	if (Count < 0)
+	{
+		Count = GetNumPoints();
+	}
+
+	Count = FMath::Min(Count, GetNumPoints() - StartIndex);
+
+	// TArray only supports int32 ax max number of elements!
+	check(Count <= INT32_MAX);
+
+	Points.Empty(Count);
+
+	FScopeLock Lock(&DataLock);
+
+	TQueue<const FLidarPointCloudOctreeNode*> Nodes;
+	const FLidarPointCloudOctreeNode* CurrentNode = nullptr;
+	Nodes.Enqueue(&Root);
+	while (Nodes.Dequeue(CurrentNode))
+	{
+		// If no data is required, quit
+		if (Count == 0)
+		{
+			return;
+		}
+
+		// Should this node's points be injected?
+		if (StartIndex < CurrentNode->GetNumPoints())
+		{
+			// If the index is at 0 and the requested count is at least the size of this node, append whole arrays
+			if (StartIndex == 0 && Count >= CurrentNode->GetNumPoints())
+			{
+				Points.Append(CurrentNode->AllocatedPoints);
+				Points.Append(CurrentNode->PaddingPoints);
+
+				Count -= CurrentNode->GetNumPoints();
+			}
+			// ... otherwise iterate over needed data
+			else
+			{
+				// AllocatedPoints
+				{
+					int32 NumPointsToCopy = FMath::Max(0LL, FMath::Min((int64)CurrentNode->AllocatedPoints.Num(), Count + StartIndex) - StartIndex);
+					if (NumPointsToCopy > 0)
+					{
+						int32 SrcOffset = StartIndex;
+						int32 DestOffset = Points.Num();
+
+						// Expand the array
+						Points.AddUninitialized(NumPointsToCopy);
+
+						// Copy contents
+						FMemory::Memcpy(Points.GetData() + DestOffset, CurrentNode->AllocatedPoints.GetData() + SrcOffset, NumPointsToCopy * sizeof(FLidarPointCloudPoint));
+
+						StartIndex = FMath::Max(0LL, StartIndex - NumPointsToCopy);
+						Count -= NumPointsToCopy;
+					}
+				}
+
+				// PaddingPoints
+				{
+					int32 NumPointsToCopy = FMath::Max(0LL, FMath::Min((int64)CurrentNode->PaddingPoints.Num(), Count + StartIndex) - StartIndex);
+					if (NumPointsToCopy > 0)
+					{
+						int32 SrcOffset = StartIndex;
+						int32 DestOffset = Points.Num();
+
+						// Expand the array
+						Points.AddUninitialized(NumPointsToCopy);
+
+						// Copy contents
+						FMemory::Memcpy(Points.GetData() + DestOffset, CurrentNode->PaddingPoints.GetData() + SrcOffset, NumPointsToCopy * sizeof(FLidarPointCloudPoint));
+
+						StartIndex = FMath::Max(0LL, StartIndex - NumPointsToCopy);;
+						Count -= NumPointsToCopy;
+					}
+				}
+			}
+		}
+		// ... or skipped completely?
+		else
+		{
+			StartIndex -= CurrentNode->GetNumPoints();
+		}
+
+		for (auto Child : CurrentNode->Children)
+		{
+			Nodes.Enqueue(Child);
+		}
+	}
+}
+
+void FLidarPointCloudOctree::GetPointsInSphereAsCopies(TArray<FLidarPointCloudPoint>& SelectedPoints, const FSphere& Sphere, const bool& bVisibleOnly) const
+{
+	SelectedPoints.Reset();
+	PROCESS_IN_SPHERE_CONST({ SelectedPoints.Add(*Point); });
+}
+
+void FLidarPointCloudOctree::GetPointsInBoxAsCopies(TArray<FLidarPointCloudPoint>& SelectedPoints, const FBox& Box, const bool& bVisibleOnly) const
+{
+	SelectedPoints.Reset();
+	PROCESS_IN_BOX_CONST({ SelectedPoints.Add(*Point); });
+}
+
+FLidarPointCloudPoint* FLidarPointCloudOctree::RaycastSingle(const FLidarPointCloudRay& Ray, const float& Radius, const bool& bVisibleOnly)
+{
+	PROCESS_BY_RAY_CONST({ return Point; });
+	return nullptr;
+}
+
+bool FLidarPointCloudOctree::RaycastMulti(const FLidarPointCloudRay& Ray, const float& Radius, const bool& bVisibleOnly, TArray<FLidarPointCloudPoint*>& OutHits)
+{
+	OutHits.Reset();
+	PROCESS_BY_RAY_CONST({ OutHits.Add(Point); });
+	return OutHits.Num() > 0;
+}
+
+bool FLidarPointCloudOctree::RaycastMulti(const FLidarPointCloudRay& Ray, const float& Radius, const bool& bVisibleOnly, TArray<FLidarPointCloudPoint>& OutHits)
+{
+	OutHits.Reset();
+	PROCESS_BY_RAY_CONST({ OutHits.Add(*Point); });
+	return OutHits.Num() > 0;
+}
+
+bool FLidarPointCloudOctree::HasPointsInSphere(const FSphere& Sphere, const bool& bVisibleOnly) const
+{
+	PROCESS_IN_SPHERE_CONST({ return true; });
+	return false;
+}
+
+bool FLidarPointCloudOctree::HasPointsInBox(const FBox& Box, const bool& bVisibleOnly) const
+{
+	PROCESS_IN_BOX_CONST({ return true; });
+	return false;
+}
+
+void FLidarPointCloudOctree::SetVisibilityOfPointsInSphere(const bool& bNewVisibility, const FSphere& Sphere)
+{
+	// Build a box to quickly filter out the points - (IsInsideOrOn vs comparing DistSquared)
+	FBox Box(Sphere.Center - FVector(Sphere.W), Sphere.Center + FVector(Sphere.W));
+	float RadiusSq = Sphere.W * Sphere.W;
+
+	ITERATE_NODES({
+		// Skip node if it already has all points set to the required visibility state
+		bool bSkipNode = (CurrentNode->NumVisiblePoints == CurrentNode->GetNumPoints() && bNewVisibility) || (CurrentNode->NumVisiblePoints == 0 && !bNewVisibility);
+		if (!bSkipNode)
+		{
+			CurrentNode->NumVisiblePoints = 0;
+
+			// If node fully inside the radius - do not check individual points
+			if (CurrentNode->GetSphereBounds(this).IsInside(Sphere))
+			{
+				ITERATE_POINTS({ Point->bVisible = bNewVisibility; });
+
+				if (bNewVisibility)
+				{
+					CurrentNode->NumVisiblePoints = CurrentNode->GetNumPoints();
+				}
+			}
+			else
+			{
+				ITERATE_POINTS(
+				{
+					if (Point->bVisible != bNewVisibility && POINT_IN_SPHERE)
+					{
+						Point->bVisible = bNewVisibility;
+					}
+
+					if (Point->bVisible)
+					{
+						++CurrentNode->NumVisiblePoints;
+					}
+				});
+			}
+
+			CurrentNode->bVisibilityDirty = false;
+
+			// Sort points to speed up rendering
+			Algo::Sort(CurrentNode->AllocatedPoints, MoveTemp(PointVisibilitySortingFunction));
+			Algo::Sort(CurrentNode->PaddingPoints, MoveTemp(PointVisibilitySortingFunction));
+		}
+	}, NODE_IN_BOX);
+}
+
+void FLidarPointCloudOctree::SetVisibilityOfPointsInSphereAsync(const bool& bNewVisibility, const FSphere& Sphere, TFunction<void(void)> CompletionCallback)
+{
+	Async(EAsyncExecution::Thread, [this, bNewVisibility, Sphere]
+	{
+		SetVisibilityOfPointsInSphere(bNewVisibility, Sphere);
+	}, MoveTemp(CompletionCallback));
+}
+
+void FLidarPointCloudOctree::SetVisibilityOfPointsInBox(const bool& bNewVisibility, const FBox& Box)
+{
+	ITERATE_NODES({
+		// Skip node if it already has all points set to the required visibility state
+		bool bSkipNode = (CurrentNode->NumVisiblePoints == CurrentNode->GetNumPoints() && bNewVisibility) || (CurrentNode->NumVisiblePoints == 0 && !bNewVisibility);
+		if (!bSkipNode)
+		{
+			CurrentNode->NumVisiblePoints = 0;
+
+			// If node fully inside the radius - do not check individual points
+			if (Box.IsInsideOrOn(CurrentNode->Center - SharedData[CurrentNode->Depth].Extent) && Box.IsInsideOrOn(CurrentNode->Center + SharedData[CurrentNode->Depth].Extent))
+			{
+				ITERATE_POINTS({ Point->bVisible = bNewVisibility; });
+
+				if (bNewVisibility)
+				{
+					CurrentNode->NumVisiblePoints = CurrentNode->GetNumPoints();
+				}
+			}
+			else
+			{
+				ITERATE_POINTS(
+				{
+					if (Point->bVisible != bNewVisibility && POINT_IN_BOX)
+					{
+						Point->bVisible = bNewVisibility;
+					}
+
+					if (Point->bVisible)
+					{
+						++CurrentNode->NumVisiblePoints;
+					}
+				});
+			}
+
+			CurrentNode->bVisibilityDirty = false;
+
+			// Sort points to speed up rendering
+			Algo::Sort(CurrentNode->AllocatedPoints, MoveTemp(PointVisibilitySortingFunction));
+			Algo::Sort(CurrentNode->PaddingPoints, MoveTemp(PointVisibilitySortingFunction));
+		}
+	}, NODE_IN_BOX);
+}
+
+void FLidarPointCloudOctree::SetVisibilityOfPointsInBoxAsync(const bool& bNewVisibility, const FBox& Box, TFunction<void(void)> CompletionCallback)
+{
+	Async(EAsyncExecution::Thread, [this, bNewVisibility, Box]
+	{
+		SetVisibilityOfPointsInBox(bNewVisibility, Box);
+	}, MoveTemp(CompletionCallback));
+}
+
+void FLidarPointCloudOctree::SetVisibilityOfPointsByRay(const bool& bNewVisibility, const FLidarPointCloudRay& Ray, const float& Radius)
+{
+	const float RadiusSq = Radius * Radius;
+	ITERATE_NODES({
+		// Skip node if it already has all points set to the required visibility state
+		bool bSkipNode = (CurrentNode->NumVisiblePoints == CurrentNode->GetNumPoints() && bNewVisibility) || (CurrentNode->NumVisiblePoints == 0 && !bNewVisibility);
+		if (!bSkipNode && Ray.Intersects(CurrentNode->GetBounds(this)))
+		{
+			CurrentNode->NumVisiblePoints = 0;
+
+			ITERATE_POINTS(
+			{
+				if (Point->bVisible != bNewVisibility && POINT_BY_RAY)
+				{
+					Point->bVisible = bNewVisibility;
+				}
+
+				if (Point->bVisible)
+				{
+					++CurrentNode->NumVisiblePoints;
+				}
+			});
+
+			CurrentNode->bVisibilityDirty = false;
+
+			// Sort points to speed up rendering
+			Algo::Sort(CurrentNode->AllocatedPoints, MoveTemp(PointVisibilitySortingFunction));
+			Algo::Sort(CurrentNode->PaddingPoints, MoveTemp(PointVisibilitySortingFunction));
+
+			for (auto Child : CurrentNode->Children)
+			{
+				Nodes.Enqueue(Child);
+			}
+		}
+	}, false);
+}
+
+void FLidarPointCloudOctree::SetVisibilityOfPointsByRayAsync(const bool& bNewVisibility, const FLidarPointCloudRay& Ray, const float& Radius, TFunction<void(void)> CompletionCallback /*= nullptr*/)
+{
+	Async(EAsyncExecution::Thread, [this, bNewVisibility, Ray, Radius]
+	{
+		SetVisibilityOfPointsByRay(bNewVisibility, Ray, Radius);
+	}, MoveTemp(CompletionCallback));
+}
+
+void FLidarPointCloudOctree::HideAll()
+{
+	ITERATE_NODES({
+		if (CurrentNode->NumVisiblePoints > 0)
+		{
+			ITERATE_POINTS({ Point->bVisible = false; });
+
+			CurrentNode->NumVisiblePoints = 0;
+			CurrentNode->bVisibilityDirty = false;
+		}
+		}, true);
+}
+
+void FLidarPointCloudOctree::UnhideAll()
+{
+	ITERATE_NODES({
+		if (CurrentNode->NumVisiblePoints != CurrentNode->GetNumPoints())
+		{
+			ITERATE_POINTS({ Point->bVisible = true; });
+
+			CurrentNode->NumVisiblePoints = CurrentNode->GetNumPoints();
+			CurrentNode->bVisibilityDirty = false;
+		}
+	}, true);
+}
+
+void FLidarPointCloudOctree::ResetVisibilityAsync(TFunction<void(void)> CompletionCallback)
+{
+	Async(EAsyncExecution::Thread, [this]
+	{
+		ResetVisibilityAsync();
+	}, MoveTemp(CompletionCallback));
+}
+
+void FLidarPointCloudOctree::ExecuteActionOnAllPoints(TFunction<void(FLidarPointCloudPoint*)> Action, const bool& bVisibleOnly)
+{
+	PROCESS_ALL({ Action(Point); });
+}
+
+void FLidarPointCloudOctree::ExecuteActionOnAllPointsAsync(TFunction<void(FLidarPointCloudPoint*)> Action, const bool& bVisibleOnly, TFunction<void(void)> CompletionCallback /*= nullptr*/)
+{
+	Async(EAsyncExecution::Thread, [this, Action, bVisibleOnly]
+	{
+		ExecuteActionOnAllPointsAsync(Action, bVisibleOnly);
+	}, MoveTemp(CompletionCallback));
+}
+
+void FLidarPointCloudOctree::ExecuteActionOnPointsInSphere(TFunction<void(FLidarPointCloudPoint*)> Action, const FSphere& Sphere, const bool& bVisibleOnly)
+{
+	PROCESS_IN_SPHERE({ Action(Point); });
+}
+
+void FLidarPointCloudOctree::ExecuteActionOnPointsInSphereAsync(TFunction<void(FLidarPointCloudPoint*)> Action, const FSphere& Sphere, const bool& bVisibleOnly, TFunction<void(void)> CompletionCallback)
+{
+	Async(EAsyncExecution::Thread, [this, Action, Sphere, bVisibleOnly]
+	{
+		ExecuteActionOnPointsInSphere(Action, Sphere, bVisibleOnly);
+	}, MoveTemp(CompletionCallback));
+}
+
+void FLidarPointCloudOctree::ExecuteActionOnPointsInBox(TFunction<void(FLidarPointCloudPoint*)> Action, const FBox& Box, const bool& bVisibleOnly)
+{
+	PROCESS_IN_BOX({ Action(Point); });
+}
+
+void FLidarPointCloudOctree::ExecuteActionOnPointsInBoxAsync(TFunction<void(FLidarPointCloudPoint*)> Action, const FBox& Box, const bool& bVisibleOnly, TFunction<void(void)> CompletionCallback)
+{
+	Async(EAsyncExecution::Thread, [this, Action, Box, bVisibleOnly]
+	{
+		ExecuteActionOnPointsInBox(Action, Box, bVisibleOnly);
+	}, MoveTemp(CompletionCallback));
+}
+
+void FLidarPointCloudOctree::ExecuteActionOnPointsByRay(TFunction<void(FLidarPointCloudPoint*)> Action, const FLidarPointCloudRay& Ray, const float& Radius, bool bVisibleOnly)
+{
+	PROCESS_BY_RAY({ Action(Point); });
+}
+
+void FLidarPointCloudOctree::ExecuteActionOnPointsByRayAsync(TFunction<void(FLidarPointCloudPoint*)> Action, const FLidarPointCloudRay& Ray, const float& Radius, bool bVisibleOnly, TFunction<void(void)> CompletionCallback)
+{
+	Async(EAsyncExecution::Thread, [this, Action, Ray, Radius, bVisibleOnly]
+	{
+		ExecuteActionOnPointsByRay(Action, Ray, Radius, bVisibleOnly);
+	}, MoveTemp(CompletionCallback));
+}
+
+void FLidarPointCloudOctree::ApplyColorToAllPoints(const FColor& NewColor, const bool& bVisibleOnly)
+{
+	PROCESS_ALL({ Point->Color = NewColor; });
+}
+
+void FLidarPointCloudOctree::ApplyColorToPointsInSphere(const FColor& NewColor, const FSphere& Sphere, const bool& bVisibleOnly)
+{
+	PROCESS_IN_SPHERE({ Point->Color = NewColor; });
+}
+
+void FLidarPointCloudOctree::ApplyColorToPointsInBox(const FColor& NewColor, const FBox& Box, const bool& bVisibleOnly)
+{
+	PROCESS_IN_BOX({ Point->Color = NewColor; });
+}
+
+void FLidarPointCloudOctree::ApplyColorToPointsByRay(const FColor& NewColor, const FLidarPointCloudRay& Ray, const float& Radius, bool bVisibleOnly)
+{
+	PROCESS_BY_RAY({ Point->Color = NewColor; });
+}
+
+void FLidarPointCloudOctree::MarkPointVisibilityDirty()
+{
+	TQueue<FLidarPointCloudOctreeNode*> Nodes;
+	FLidarPointCloudOctreeNode* CurrentNode = nullptr;
+	Nodes.Enqueue(&Root);
+	while (Nodes.Dequeue(CurrentNode))
+	{
+		CurrentNode->bVisibilityDirty = true;
+
+		for (auto Child : CurrentNode->Children)
+		{
+			Nodes.Enqueue(Child);
+		}
+	}
+}
+
+void FLidarPointCloudOctree::Initialize(const FBox& InBounds)
+{
+	if (!SetNewBounds(InBounds))
+	{
+		return;
+	}
+
+	MaxBucketSize = GetDefault<ULidarPointCloudSettings>()->MaxBucketSize;
+	NodeGridResolution = GetDefault<ULidarPointCloudSettings>()->NodeGridResolution;
+
+	// Pre-calculate the shared per LOD data
+	for (int32 i = 0; i < SharedData.Num(); i++)
+	{
+		SharedData[i] = FSharedLODData(UniformBounds.GetExtent() / FMath::Pow(2, i));
+		NodeCount[i].Reset();
+		PointCount[i].Reset();
+	}
+
+	Empty(true);
+}
+
+void FLidarPointCloudOctree::ShiftPointsBy(FDoubleVector Offset, bool bRefreshPointsBounds)
+{
+	TQueue<FLidarPointCloudOctreeNode*> Nodes;
+	FLidarPointCloudOctreeNode* CurrentNode = nullptr;
+	Nodes.Enqueue(&Root);
+	while (Nodes.Dequeue(CurrentNode))
+	{
+		CurrentNode->Center = (Offset + CurrentNode->Center).ToVector();
+
+		FOR(Point, CurrentNode->AllocatedPoints)
+		{
+			Point->Location = (Offset + Point->Location).ToVector();
+		}
+
+		FOR(Point, CurrentNode->PaddingPoints)
+		{
+			Point->Location = (Offset + Point->Location).ToVector();
+		}
+
+		for (auto Child : CurrentNode->Children)
+		{
+			Nodes.Enqueue(Child);
+		}
+	}
+
+	Bounds.Min = (Offset + Bounds.Min).ToVector();
+	Bounds.Max = (Offset + Bounds.Max).ToVector();
+	UniformBounds.Min = (Offset + UniformBounds.Min).ToVector();
+	UniformBounds.Max = (Offset + UniformBounds.Max).ToVector();
+
+	MarkTraversalOctreesForInvalidation();
+
+	if (bRefreshPointsBounds)
+	{
+		RefreshPointsBounds();
+	}
+}
+
+void FLidarPointCloudOctree::RemovePoint(const FLidarPointCloudPoint* Point, bool bRefreshPointsBounds)
+{
+	if (!Point)
+	{
+		return;
+	}
+	
+	int32 Index = -1;
+	bool bAllocatedPoint = false;
+	FLidarPointCloudOctreeNode* CurrentNode = &Root;
+	while (CurrentNode)
+	{
+		// Scan Allocated Points
+		{
+			const FLidarPointCloudPoint* Start = CurrentNode->AllocatedPoints.GetData();
+			const int32 NumElements = CurrentNode->AllocatedPoints.Num();
+			for (const FLidarPointCloudPoint* Data = Start, *DataEnd = Start + NumElements; Data != DataEnd; ++Data)
+			{
+				if (Data == Point)
+				{
+					Index = Data - Start;
+					bAllocatedPoint = true;
+					break;
+				}
+			}
+		}
+
+		// Scan Padding Points
+		{
+			const FLidarPointCloudPoint* Start = CurrentNode->PaddingPoints.GetData();
+			const int32 NumElements = CurrentNode->PaddingPoints.Num();
+			for (const FLidarPointCloudPoint* Data = Start, *DataEnd = Start + NumElements; Data != DataEnd; ++Data)
+			{
+				if (Data == Point)
+				{
+					Index = Data - Start;
+					break;
+				}
+			}
+		}
+
+		if (Index > 0)
+		{
+			RemovePoint_Internal(CurrentNode, Index, bAllocatedPoint);
+			break;
+		}
+		else
+		{
+			FVector CenterRelativeLocation = Point->Location - CurrentNode->Center;
+			CurrentNode = CurrentNode->GetChildNodeAtLocation((CenterRelativeLocation.X > 0 ? 4 : 0) + (CenterRelativeLocation.Y > 0 ? 2 : 0) + (CenterRelativeLocation.Z > 0));
+		}
+	}
+
+	if (bRefreshPointsBounds)
+	{
+		RefreshPointsBounds();
+	}
+}
+
+void FLidarPointCloudOctree::RemovePoint(FLidarPointCloudPoint Point, bool bRefreshPointsBounds)
+{
+	bool bAllocatedPoint = false;
+	FLidarPointCloudOctreeNode* CurrentNode = &Root;
+	while (CurrentNode)
+	{
+		int32 Index = CurrentNode->AllocatedPoints.IndexOfByKey(Point);
+		if (Index != INDEX_NONE)
+		{
+			bAllocatedPoint = true;
+		}
+		else
+		{
+			Index = CurrentNode->PaddingPoints.IndexOfByKey(Point);
+		}
+
+		if (Index != INDEX_NONE)
+		{
+			RemovePoint_Internal(CurrentNode, Index, bAllocatedPoint);
+			break;
+		}
+		else
+		{
+			FVector CenterRelativeLocation = Point.Location - CurrentNode->Center;
+			CurrentNode = CurrentNode->GetChildNodeAtLocation((CenterRelativeLocation.X > 0 ? 4 : 0) + (CenterRelativeLocation.Y > 0 ? 2 : 0) + (CenterRelativeLocation.Z > 0));
+		}
+	}
+
+	if (bRefreshPointsBounds)
+	{
+		RefreshPointsBounds();
+	}
+}
+
+void FLidarPointCloudOctree::RemovePoints(TArray<FLidarPointCloudPoint*>& Points, bool bRefreshPointsBounds)
+{
+	if (Points.Num() == 0)
+	{
+		return;
+	}
+
+	FOR(Point, Points)
+	{
+		(*Point)->bMarkedForDeletion = true;
+	}
+
+	ITERATE_NODES(
+	{
+		int64 NumRemoved = 0;
+
+		// Scan Allocated Points
+		{
+			FLidarPointCloudPoint* Start = CurrentNode->AllocatedPoints.GetData();
+			const int32 NumElements = CurrentNode->AllocatedPoints.Num();
+			for (FLidarPointCloudPoint* Data = Start, *DataEnd = Start + NumElements; Data != DataEnd; ++Data)
+			{
+				if (Data->bMarkedForDeletion)
+				{
+					CurrentNode->AllocatedPoints.RemoveAtSwap(Data - Start, 1, false);
+					++NumRemoved;
+					--DataEnd; 
+					--Data;
+				}
+			}
+		}
+
+		// Scan Padding Points
+		{
+			FLidarPointCloudPoint* Start = CurrentNode->PaddingPoints.GetData();
+			const int32 NumElements = CurrentNode->PaddingPoints.Num();
+			for (FLidarPointCloudPoint* Data = Start, *DataEnd = Start + NumElements; Data != DataEnd; ++Data)
+			{
+				if (Data->bMarkedForDeletion)
+				{
+					CurrentNode->PaddingPoints.RemoveAtSwap(Data - Start, 1, false);
+					++NumRemoved;
+					--DataEnd;
+					--Data;
+				}
+			}
+		}
+
+		if (NumRemoved > 0)
+		{
+			// #todo: Fetch points from child nodes / padding points to fill the gap
+			
+			CurrentNode->AddPointCount(this, -NumRemoved);
+
+			// Reduce space usage
+			CurrentNode->AllocatedPoints.Shrink();
+			CurrentNode->PaddingPoints.Shrink();
+
+			// Sort points to speed up rendering
+			Algo::Sort(CurrentNode->AllocatedPoints, MoveTemp(PointVisibilitySortingFunction));
+			Algo::Sort(CurrentNode->PaddingPoints, MoveTemp(PointVisibilitySortingFunction));
+		}
+	}, true);
+
+	if (bRefreshPointsBounds)
+	{
+		RefreshPointsBounds();
+	}
+}
+
+void FLidarPointCloudOctree::RemovePointsInSphere(const FSphere& Sphere, const bool& bVisibleOnly, bool bRefreshPointsBounds)
+{
+	// #todo: This can be optimized by removing points inline
+	TArray<FLidarPointCloudPoint*> SelectedPoints;
+	GetPointsInSphere(SelectedPoints, Sphere, bVisibleOnly);
+	RemovePoints(SelectedPoints, bRefreshPointsBounds);
+}
+
+void FLidarPointCloudOctree::RemovePointsInBox(const FBox& Box, const bool& bVisibleOnly, bool bRefreshPointsBounds)
+{
+	// #todo: This can be optimized by removing points inline
+	TArray<FLidarPointCloudPoint*> SelectedPoints;
+	GetPointsInBox(SelectedPoints, Box, bVisibleOnly);
+	RemovePoints(SelectedPoints, bRefreshPointsBounds);
+}
+
+void FLidarPointCloudOctree::RemovePointsByRay(const FLidarPointCloudRay& Ray, const float& Radius, const bool& bVisibleOnly, bool bRefreshPointsBounds)
+{
+	// #todo: This can be optimized by removing points inline
+	TArray<FLidarPointCloudPoint*> SelectedPoints;
+	RaycastMulti(Ray, Radius, bVisibleOnly, SelectedPoints);
+	RemovePoints(SelectedPoints, bRefreshPointsBounds);
+}
+
+void FLidarPointCloudOctree::RemoveHiddenPoints(bool bRefreshPointsBounds)
+{
+	ITERATE_NODES(
+	{
+		int64 NumRemoved = 0;
+
+		// Scan Allocated Points
+		{
+			FLidarPointCloudPoint* Start = CurrentNode->AllocatedPoints.GetData();
+			const int32 NumElements = CurrentNode->AllocatedPoints.Num();
+			for (FLidarPointCloudPoint* Data = Start, *DataEnd = Start + NumElements; Data != DataEnd; ++Data)
+			{
+				if (!Data->bVisible)
+				{
+					CurrentNode->AllocatedPoints.RemoveAtSwap(Data - Start, 1, false);
+					++NumRemoved;
+					--DataEnd; 
+					--Data;
+				}
+			}
+		}
+
+		// Scan Padding Points
+		{
+			FLidarPointCloudPoint* Start = CurrentNode->PaddingPoints.GetData();
+			const int32 NumElements = CurrentNode->PaddingPoints.Num();
+			for (FLidarPointCloudPoint* Data = Start, *DataEnd = Start + NumElements; Data != DataEnd; ++Data)
+			{
+				if (!Data->bVisible)
+				{
+					CurrentNode->PaddingPoints.RemoveAtSwap(Data - Start, 1, false);
+					++NumRemoved;
+					--DataEnd;
+					--Data;
+				}
+			}
+		}
+
+		if (NumRemoved > 0)
+		{
+			// #todo: Fetch points from child nodes / padding points to fill the gap
+			
+			CurrentNode->AddPointCount(this, -NumRemoved);
+
+			// Reduce space usage
+			CurrentNode->AllocatedPoints.Shrink();
+			CurrentNode->PaddingPoints.Shrink();
+
+			// Set visibility data
+			CurrentNode->NumVisiblePoints = CurrentNode->GetNumPoints();
+			CurrentNode->bVisibilityDirty = false;
+		}
+	}, true);
+
+	if (bRefreshPointsBounds)
+	{
+		RefreshPointsBounds();
+	}
+}
+
+void FLidarPointCloudOctree::Empty(bool bDestroyNodes)
+{
+	if (bDestroyNodes)
+	{
+		Root.~FLidarPointCloudOctreeNode();
+
+		// Reset node counters
+		for (auto& Count : NodeCount)
+		{
+			Count.Reset();
+		}
+
+		new (&Root) FLidarPointCloudOctreeNode(this, nullptr, 0, 0, Bounds.GetCenter());
+
+		MarkTraversalOctreesForInvalidation();
+	}
+	else
+	{
+		Root.Empty(true);
+	}
+
+	PointsBounds.Init();
+
+	// Reset point counters
+	for (auto& Count : PointCount)
+	{
+		Count.Reset();
+	}
+}
+
+bool FLidarPointCloudOctree::SetNewBounds(const FBox& InBounds)
+{
+	if (!InBounds.IsValid)
+	{
+		return false;
+	}
+
+	Bounds = InBounds;
+
+	FVector BoundsExtent = FVector(Bounds.GetExtent().GetMax());
+	UniformBounds = FBox(Bounds.GetCenter() - BoundsExtent, Bounds.GetCenter() + BoundsExtent);
+
+	return true;
+}
+
+void FLidarPointCloudOctree::RemovePoint_Internal(FLidarPointCloudOctreeNode* Node, int32 Index, bool bAllocatedPoint)
+{
+	Node->AddPointCount(this, -1);
+
+	// Remove the point from relevant array
+	(bAllocatedPoint ? Node->AllocatedPoints : Node->PaddingPoints).RemoveAt(Index);
+
+	// #todo: Fetch points from child nodes / padding points to fill the gap
+}
+
+void FLidarPointCloudOctree::MarkTraversalOctreesForInvalidation()
+{
+	for (int32 i = 0; i < LinkedTraversalOctrees.Num(); ++i)
+	{
+		auto TraversalOctree = LinkedTraversalOctrees[i];
+
+		// Remove null
+		if (!TraversalOctree)
+		{
+			LinkedTraversalOctrees.RemoveAtSwap(i--);
+		}
+		else if(TraversalOctree->bValid)
+		{
+			TraversalOctree->bValid = false;
+		}
+	}
+}
+
+FArchive& operator<<(FArchive& Ar, FLidarPointCloudOctree& O)
+{
+	Ar << O.Bounds;
+
+	// Collision Mesh data
+	if (Ar.CustomVer(ULidarPointCloud::PointCloudFileGUID) > 13)
+	{
+		FTriMeshCollisionData Dummy;
+		FTriMeshCollisionData* CollisionMesh = Ar.IsCooking() ? &Dummy : &O.CollisionMesh;
+
+		Ar << CollisionMesh->Vertices;
+
+		int32 NumIndices = CollisionMesh->Indices.Num();
+		Ar << NumIndices;
+
+		if (Ar.IsLoading())
+		{
+			CollisionMesh->Indices.AddUninitialized(NumIndices);
+		}
+
+		Ar.Serialize(CollisionMesh->Indices.GetData(), NumIndices * sizeof(FTriIndices));
+	}
+
+	if (Ar.IsLoading())
+	{
+		O.Initialize(O.Bounds);
+	}
+
+	O.Root.Serialize(Ar, &O);
+
+	// Points Bounds
+	if (Ar.CustomVer(ULidarPointCloud::PointCloudFileGUID) > 14)
+	{
+		Ar << O.PointsBounds;
+	}
+	else
+	{
+		O.RefreshPointsBounds();
+	}
+
+	return Ar;
+}
+
+//////////////////////////////////////////////////////////// FLidarPointCloudTraversalOctreeNode
+
+FLidarPointCloudTraversalOctreeNode::FLidarPointCloudTraversalOctreeNode()
+	: DataNode(nullptr)
+	, Parent(nullptr)
+	, Children(nullptr)
+{
+}
+
+FLidarPointCloudTraversalOctreeNode::~FLidarPointCloudTraversalOctreeNode()
+{
+	if (Children)
+	{
+		delete[] Children;
+	}
+}
+
+void FLidarPointCloudTraversalOctreeNode::Build(FLidarPointCloudOctreeNode* Node, const FTransform& LocalToWorld)
+{
+	DataNode = Node;
+	Center = LocalToWorld.TransformPosition(Node->Center);
+	Depth = Node->Depth;
+
+	NumChildren = Node->Children.Num();
+	if (NumChildren > 0)
+	{
+		Children = new FLidarPointCloudTraversalOctreeNode[NumChildren];
+
+		for (int32 i = 0; i < NumChildren; i++)
+		{
+			Children[i].Build(Node->Children[i], LocalToWorld);
+			Children[i].Parent = this;
+		}
+	}
+}
+
+void FLidarPointCloudTraversalOctreeNode::CalculateVirtualDepth(const TArray<float>& LevelWeights, const float& VDMultiplier, const float& PointSizeBias)
+{
+	if (!bSelected)
+	{
+		return;
+	}
+
+	TQueue<const FLidarPointCloudTraversalOctreeNode*> Nodes;
+	const FLidarPointCloudTraversalOctreeNode* CurrentNode = nullptr;
+
+	// Calculate virtual depth factor
+	float VDFactor = 0;
+	Nodes.Enqueue(this);
+	while (Nodes.Dequeue(CurrentNode))
+	{
+		for (int32 i = 0; i < CurrentNode->NumChildren; i++)
+		{
+			if (CurrentNode->Children[i].bSelected)
+			{
+				Nodes.Enqueue(&CurrentNode->Children[i]);
+			}
+		}
+
+		float LocalVDFactor = CurrentNode->Depth * CurrentNode->DataNode->GetNumPoints() * LevelWeights[CurrentNode->Depth];
+
+		if (CurrentNode != this && PointSizeBias > 0)
+		{
+			LocalVDFactor /= (CurrentNode->Parent->NumChildren - 1) * PointSizeBias + 1;
+		}
+
+		VDFactor += LocalVDFactor;
+	}
+
+	// Calculate weighted number of visible points
+	float NumPoints = 0;
+	Nodes.Enqueue(this);
+	while (Nodes.Dequeue(CurrentNode))
+	{
+		for (int32 i = 0; i < CurrentNode->NumChildren; i++)
+		{
+			if (CurrentNode->Children[i].bSelected)
+			{
+				Nodes.Enqueue(&CurrentNode->Children[i]);
+			}
+		}
+
+		NumPoints += CurrentNode->DataNode->GetNumPoints() * LevelWeights[CurrentNode->Depth];
+	}
+
+	// Calculate the Virtual Depth
+	VirtualDepth = VDFactor / NumPoints * VDMultiplier;
+}
+
+//////////////////////////////////////////////////////////// FLidarPointCloudTraversalOctree
+
+FLidarPointCloudTraversalOctree::FLidarPointCloudTraversalOctree(FLidarPointCloudOctree* Octree, const FTransform& LocalToWorld)
+	: Octree(Octree)
+	, bValid(false)
+{
+	Octree->RegisterTraversalOctree(this);
+	
+	// Reset the tree
+	Extents.Empty();
+	RadiiSq.Empty();
+	PointCount.Empty();
+	NumPoints = 0;
+	Root.~FLidarPointCloudTraversalOctreeNode();
+	new (&Root) FLidarPointCloudTraversalOctreeNode();
+
+	// Calculate properties
+	NumLODs = Octree->GetNumLODs();
+
+	FBox WorldBounds = Octree->UniformBounds.TransformBy(LocalToWorld);
+	for (int32 i = 0; i < NumLODs; i++)
+	{
+		Extents.Emplace(i == 0 ? WorldBounds.GetExtent() : Extents.Last() * 0.5f);
+		RadiiSq.Emplace(FMath::Square(Extents.Last().Size()));
+	}
+
+	for (auto& Count : Octree->PointCount)
+	{
+		PointCount.Add(Count.GetValue());
+		NumPoints += Count.GetValue();
+	}
+
+	// Star cloning the node data
+	Root.Build(&Octree->Root, LocalToWorld);
+
+	bValid = true;
+}
+
+FLidarPointCloudTraversalOctree::~FLidarPointCloudTraversalOctree()
+{
+	Octree->UnregisterTraversalOctree(this);
+}
+
+TArray<FLidarPointCloudTraversalOctreeNode*> FLidarPointCloudTraversalOctree::GetVisibleNodes(const FSceneView* View, const float& MinScreenSizeSq, const int32& PointBudget, const float& ScreenCenterImportance, const int32& MinDepth, const int32& MaxDepth)
+{
+	SCOPE_CYCLE_COUNTER(STAT_NodeSelection);
+
+	// For Top / Front / Side views disable ScreenSize checks
+	const bool bSkipMinScreenSize = !View->bIsGameView && !View->IsPerspectiveProjection();
+
+	// Collect Nodes
+	TArray<FNodeSizeData> NodeSizeData;
+	{
+		const FMatrix& ProjMatrix = View->ViewMatrices.GetProjectionMatrix();
+		const float ScreenSizeFactor = FMath::Square(FMath::Max(0.5f * ProjMatrix.M[0][0], 0.5f * ProjMatrix.M[1][1]));
+		const FVector& ViewOrigin = View->ViewMatrices.GetViewOrigin();
+
+		TQueue<FLidarPointCloudTraversalOctreeNode*> Nodes;
+		FLidarPointCloudTraversalOctreeNode* CurrentNode = nullptr;
+		Nodes.Enqueue(&Root);
+		while (Nodes.Dequeue(CurrentNode))
+		{
+			// Reset selection flag
+			CurrentNode->bSelected = false;
+
+			// Update number of visible points, if needed
+			CurrentNode->DataNode->UpdateNumVisiblePoints();
+
+			// In Frustum?
+			// #todo: Skip frustum checks for nodes fully in frustum
+			if (!View->ViewFrustum.IntersectBox(CurrentNode->Center, Extents[CurrentNode->Depth]))
+			{
+				continue;
+			}
+
+			// Only process this node if it has any visible points - do not use continue; as the children may still contain visible points!
+			if (CurrentNode->DataNode->GetNumVisiblePoints() > 0 && CurrentNode->Depth >= MinDepth)
+			{
+				float ScreenSizeSq = 0;
+
+				FVector VectorToNode = CurrentNode->Center - ViewOrigin;
+				const float DistSqr = VectorToNode.SizeSquared();
+
+				// If the camera is within this node's bounds, it should always be qualified for rendering
+				if (DistSqr <= RadiiSq[CurrentNode->Depth])
+				{
+					// Subtract Depth to maintain hierarchy 
+					ScreenSizeSq = 10000 - CurrentNode->Depth;
+				}
+				else
+				{
+					ScreenSizeSq = ScreenSizeFactor * RadiiSq[CurrentNode->Depth] / FMath::Max(1.0f, DistSqr);
+
+					// Check for minimum screen size
+					if (!bSkipMinScreenSize && ScreenSizeSq < MinScreenSizeSq)
+					{
+						continue;
+					}
+
+					// Add optional preferential selection for nodes closer to the screen center
+					if (ScreenCenterImportance > 0)
+					{
+						VectorToNode.Normalize();
+						float Dot = FVector::DotProduct(View->GetViewDirection(), VectorToNode);
+
+						ScreenSizeSq = FMath::Lerp(ScreenSizeSq, ScreenSizeSq * Dot, ScreenCenterImportance);
+					}
+				}
+
+				NodeSizeData.Emplace(CurrentNode, ScreenSizeSq);
+			}
+
+			if (MaxDepth < 0 || CurrentNode->Depth < MaxDepth)
+			{
+				for (int32 i = 0; i < CurrentNode->NumChildren; i++)
+				{
+					Nodes.Enqueue(&CurrentNode->Children[i]);
+				}
+			}
+		}
+	}
+
+	// Sort Nodes
+	Algo::Sort(NodeSizeData, [](const FNodeSizeData& A, const FNodeSizeData& B) { return A.Size > B.Size; });
+
+	// Select Nodes
+	int32 TotalPointsSelected = 0;
+	TArray<FLidarPointCloudTraversalOctreeNode*> SelectedNodes;
+
+	for (FNodeSizeData& Element : NodeSizeData)
+	{
+		int32 NewNumPointsSelected = TotalPointsSelected + Element.Node->DataNode->GetNumVisiblePoints();
+
+		if (NewNumPointsSelected <= PointBudget)
+		{
+			SelectedNodes.Add(Element.Node);
+			TotalPointsSelected = NewNumPointsSelected;
+			Element.Node->bSelected = true;
+		}
+	}
+
+	return SelectedNodes;
+}
+
+void FLidarPointCloudTraversalOctree::CalculateSpriteSize(TArray<FLidarPointCloudTraversalOctreeNode*>& SelectedNodes, const float& PointSizeBias)
+{
+	for (auto Node : SelectedNodes)
+	{
+		Node->CalculateVirtualDepth(CalculateLevelWeights(), 255.0f / NumLODs, PointSizeBias);
+	}
+}
+
+TArray<float> FLidarPointCloudTraversalOctree::CalculateLevelWeights() const
+{
+	TArray<float> LevelWeights;
+	LevelWeights.AddZeroed(NumLODs);
+
+	for (int32 i = 0; i < LevelWeights.Num(); i++)
+	{
+		LevelWeights[i] = (float)PointCount[i] / NumPoints;
+	}
+
+	return LevelWeights;
+}
diff --git a/Engine/Plugins/Enterprise/LidarPointCloud/Source/LidarPointCloudRuntime/Private/LidarPointCloudOctreeMacros.h b/Engine/Plugins/Enterprise/LidarPointCloud/Source/LidarPointCloudRuntime/Private/LidarPointCloudOctreeMacros.h
new file mode 100644
--- /dev/null
+++ b/Engine/Plugins/Enterprise/LidarPointCloud/Source/LidarPointCloudRuntime/Private/LidarPointCloudOctreeMacros.h
@@ -0,0 +1,136 @@
+// Copyright Epic Games, Inc. All Rights Reserved.
+
+#pragma once
+
+#define FOR(ElemName, Array) for(auto ElemName = Array.GetData(), DataEnd = ElemName + Array.Num(); ElemName != DataEnd; ++ElemName)
+
+#define IS_VIS_CHECK_REQUIRED (bVisibleOnly && CurrentNode->NumVisiblePoints < CurrentNode->GetNumPoints())
+
+#define NODE_IN_BOX (Box.Intersect(Child->GetBounds(this)))
+#define NODE_IN_FRUSTUM (Frustum.IntersectBox(Child->Center, SharedData[Child->Depth].Extent))
+
+#define ITERATE_NODES_BODY(Action, NodeTest, Const) \
+{\
+	TQueue<Const FLidarPointCloudOctreeNode*> Nodes;\
+	Const FLidarPointCloudOctreeNode* CurrentNode = nullptr;\
+	Nodes.Enqueue(&Root);\
+	while (Nodes.Dequeue(CurrentNode))\
+	{\
+		{ Action } \
+		for (auto Child : CurrentNode->Children) if (NodeTest) Nodes.Enqueue(Child);\
+	}\
+}
+#define ITERATE_NODES(Action, NodeTest) ITERATE_NODES_BODY(Action, NodeTest,  )
+#define ITERATE_NODES_CONST(Action, NodeTest) ITERATE_NODES_BODY(Action, NodeTest, const)
+
+#define ITERATE_POINTS(Action) \
+{\
+	FOR(Point, CurrentNode->AllocatedPoints) { Action }\
+	FOR(Point, CurrentNode->PaddingPoints) { Action }\
+}
+
+#define POINT_IN_BOX Box.IsInsideOrOn(Point->Location)
+#define POINT_IN_SPHERE (POINT_IN_BOX && FVector::DistSquared(Point->Location, Sphere.Center) <= RadiusSq)
+#define POINT_IN_FRUSTUM Frustum.IntersectSphere(Point->Location, 0)
+#define POINT_BY_RAY Ray.Intersects(Point, RadiusSq)
+
+#define PROCESS_BODY(Action, PointTest) \
+{\
+	/* If node fully inside the box - do not check individual points */\
+	if (bNodeFullyContained)\
+	{\
+		if (!IS_VIS_CHECK_REQUIRED) { ITERATE_POINTS(Action) }\
+		else { ITERATE_POINTS({ if (Point->bVisible) { Action } }) }\
+	}\
+	else\
+	{\
+		if (!IS_VIS_CHECK_REQUIRED) { ITERATE_POINTS({ if (PointTest) { Action } }) }\
+		else { ITERATE_POINTS({ if (Point->bVisible && PointTest) { Action } }) }\
+	}\
+}
+
+#define PROCESS_ALL_BODY(Action) \
+{\
+	if (!bVisibleOnly || CurrentNode->NumVisiblePoints > 0)\
+	{\
+		const bool bNodeFullyContained = true;\
+		PROCESS_BODY(Action, true) \
+	}\
+}
+
+#define PROCESS_IN_SPHERE_BODY(Action) \
+{\
+	if (!bVisibleOnly || CurrentNode->NumVisiblePoints > 0)\
+	{\
+		const bool bNodeFullyContained = CurrentNode->GetSphereBounds(this).IsInside(Sphere);\
+		PROCESS_BODY(Action, POINT_IN_SPHERE) \
+	}\
+}
+
+#define PROCESS_IN_BOX_BODY(Action) \
+{\
+	if (!bVisibleOnly || CurrentNode->NumVisiblePoints > 0)\
+	{\
+		const bool bNodeFullyContained = Box.IsInsideOrOn(CurrentNode->Center - SharedData[CurrentNode->Depth].Extent) && Box.IsInsideOrOn(CurrentNode->Center + SharedData[CurrentNode->Depth].Extent);\
+		PROCESS_BODY(Action, POINT_IN_BOX) \
+	}\
+}
+
+#define PROCESS_IN_FRUSTUM_BODY(Action) \
+{\
+	if (!bVisibleOnly || CurrentNode->NumVisiblePoints > 0)\
+	{\
+		bool bNodeFullyContained;\
+		Frustum.IntersectBox(CurrentNode->Center, SharedData[CurrentNode->Depth].Extent, bNodeFullyContained);\
+		PROCESS_BODY(Action, POINT_IN_FRUSTUM) \
+	}\
+}
+
+#define PROCESS_BY_RAY_BODY(Action) \
+{\
+	if (!bVisibleOnly || CurrentNode->NumVisiblePoints > 0)\
+	{\
+		if (Ray.Intersects(CurrentNode->GetBounds(this)))\
+		{\
+			if (!IS_VIS_CHECK_REQUIRED) { ITERATE_POINTS({ if (POINT_BY_RAY) { Action } }) }\
+			else { ITERATE_POINTS({ if (Point->bVisible && POINT_BY_RAY) { Action } }) }\
+\
+			for (auto Child : CurrentNode->Children) Nodes.Enqueue(Child);\
+		}\
+	}\
+}
+
+#define PROCESS_IN_SPHERE(Action) \
+{\
+	/* Build a box to quickly filter out the points - (IsInsideOrOn vs comparing DistSquared) */\
+	const FBox Box(Sphere.Center - FVector(Sphere.W), Sphere.Center + FVector(Sphere.W));\
+	const float RadiusSq = Sphere.W * Sphere.W;\
+	ITERATE_NODES(PROCESS_IN_SPHERE_BODY(Action), NODE_IN_BOX) \
+}
+#define PROCESS_IN_SPHERE_CONST(Action) \
+{\
+	/* Build a box to quickly filter out the points - (IsInsideOrOn vs comparing DistSquared) */\
+	const FBox Box(Sphere.Center - FVector(Sphere.W), Sphere.Center + FVector(Sphere.W));\
+	const float RadiusSq = Sphere.W * Sphere.W;\
+	ITERATE_NODES_CONST(PROCESS_IN_SPHERE_BODY(Action), NODE_IN_BOX) \
+}
+
+#define PROCESS_ALL(Action) { ITERATE_NODES(PROCESS_ALL_BODY(Action), true) }
+#define PROCESS_ALL_CONST(Action) { ITERATE_NODES_CONST(PROCESS_ALL_BODY(Action), true) }
+
+#define PROCESS_IN_BOX(Action) { ITERATE_NODES(PROCESS_IN_BOX_BODY(Action), NODE_IN_BOX) }
+#define PROCESS_IN_BOX_CONST(Action) { ITERATE_NODES_CONST(PROCESS_IN_BOX_BODY(Action), NODE_IN_BOX) }
+
+#define PROCESS_IN_FRUSTUM(Action) { ITERATE_NODES(PROCESS_IN_FRUSTUM_BODY(Action), NODE_IN_FRUSTUM) }
+#define PROCESS_IN_FRUSTUM_CONST(Action) { ITERATE_NODES_CONST(PROCESS_IN_FRUSTUM_BODY(Action), NODE_IN_FRUSTUM) }
+
+#define PROCESS_BY_RAY(Action)\
+{\
+	const float RadiusSq = Radius * Radius;\
+	ITERATE_NODES(PROCESS_BY_RAY_BODY(Action), false)\
+}
+#define PROCESS_BY_RAY_CONST(Action)\
+{\
+	const float RadiusSq = Radius * Radius; \
+	ITERATE_NODES_CONST(PROCESS_BY_RAY_BODY(Action), false)\
+}
\ No newline at end of file
diff --git a/Engine/Plugins/Enterprise/LidarPointCloud/Source/LidarPointCloudRuntime/Private/LidarPointCloudSettings.cpp b/Engine/Plugins/Enterprise/LidarPointCloud/Source/LidarPointCloudRuntime/Private/LidarPointCloudSettings.cpp
new file mode 100644
--- /dev/null
+++ b/Engine/Plugins/Enterprise/LidarPointCloud/Source/LidarPointCloudRuntime/Private/LidarPointCloudSettings.cpp
@@ -0,0 +1,26 @@
+// Copyright Epic Games, Inc. All Rights Reserved.
+
+#include "LidarPointCloudSettings.h"
+#include "Math/UnrealMathUtility.h"
+
+ULidarPointCloudSettings::ULidarPointCloudSettings()
+	: DuplicateHandling(ELidarPointCloudDuplicateHandling::SelectBrighter)
+	, MaxDistanceForDuplicate(KINDA_SMALL_NUMBER)
+	, MaxBucketSize(200)
+	, NodeGridResolution(96)
+	, bUseMultithreading(true)
+	, MultithreadingInsertionBatchSize(500000)
+	, bUseAsyncImport(true)
+	, MaxImportBufferSize(52428800)
+	, ExportBatchSize(5000000)
+	, MeshingBatchSize(128)
+	, bAutoCenterOnImport(true)
+	, ImportScale(100)
+	, bEnable8BitLASDetection(true)
+	, MaxNumberOfPointsToScanLAS(1000000)
+	, MaxNumberOfPointsToScanASCII(100000)
+	, ExportScale(0.01f)
+	, bUseIOCaching(false)
+	, bLogBufferExpansion(false)
+{
+}
diff --git a/Engine/Plugins/Enterprise/LidarPointCloud/Source/LidarPointCloudRuntime/Private/LidarPointCloudShared.cpp b/Engine/Plugins/Enterprise/LidarPointCloud/Source/LidarPointCloudRuntime/Private/LidarPointCloudShared.cpp
new file mode 100644
--- /dev/null
+++ b/Engine/Plugins/Enterprise/LidarPointCloud/Source/LidarPointCloudRuntime/Private/LidarPointCloudShared.cpp
@@ -0,0 +1,33 @@
+// Copyright Epic Games, Inc. All Rights Reserved.
+ 
+#include "LidarPointCloudShared.h"
+#include "LidarPointCloud.h"
+
+double FBenchmarkTimer::Time = 0;
+
+FArchive& operator<<(FArchive& Ar, FLidarPointCloudPoint& P)
+{
+	Ar << P.Location << P.Color;
+
+	if (Ar.CustomVer(ULidarPointCloud::PointCloudFileGUID) > 8)
+	{
+		uint8 bVisible = P.bVisible;
+		Ar << bVisible;
+		P.bVisible = bVisible;
+	}
+
+	if (Ar.CustomVer(ULidarPointCloud::PointCloudFileGUID) > 12)
+	{
+		uint8 ClassificationID = P.ClassificationID;
+		Ar << ClassificationID;
+		P.ClassificationID = ClassificationID;
+	}
+
+	return Ar;
+}
+
+const FDoubleVector FDoubleVector::ZeroVector = FDoubleVector(FVector::ZeroVector);
+const FDoubleVector FDoubleVector::OneVector = FDoubleVector(FVector::OneVector);
+const FDoubleVector FDoubleVector::UpVector = FDoubleVector(FVector::UpVector);
+const FDoubleVector FDoubleVector::ForwardVector = FDoubleVector(FVector::ForwardVector);
+const FDoubleVector FDoubleVector::RightVector = FDoubleVector(FVector::RightVector);
\ No newline at end of file
diff --git a/Engine/Plugins/Enterprise/LidarPointCloud/Source/LidarPointCloudRuntime/Private/Rendering/LidarPointCloudRenderBuffers.cpp b/Engine/Plugins/Enterprise/LidarPointCloud/Source/LidarPointCloudRuntime/Private/Rendering/LidarPointCloudRenderBuffers.cpp
new file mode 100644
--- /dev/null
+++ b/Engine/Plugins/Enterprise/LidarPointCloud/Source/LidarPointCloudRuntime/Private/Rendering/LidarPointCloudRenderBuffers.cpp
@@ -0,0 +1,468 @@
+// Copyright Epic Games, Inc. All Rights Reserved.
+
+#include "LidarPointCloudRenderBuffers.h"
+#include "RenderResource.h"
+#include "VertexFactory.h"
+#include "MeshBatch.h"
+#include "RenderCommandFence.h"
+#include "LidarPointCloudOctree.h"
+#include "LidarPointCloudSettings.h"
+#include "MeshMaterialShader.h"
+
+#if WITH_EDITOR
+#include "Classes/EditorStyleSettings.h"
+#endif
+
+#define BINDPARAM(Name) Name.Bind(ParameterMap, TEXT(#Name))
+#define SETPARAM(Name) if (Name.IsBound()) { ShaderBindings.Add(Name, UserData->Name); }
+#define SETSRVPARAM(Name) if(UserData->Name) { SETPARAM(Name) }
+
+//////////////////////////////////////////////////////////// Base Buffer
+
+void FLidarPointCloudBuffer::Resize(const uint32& RequestedCapacity)
+{
+	// This must be called from Rendering thread
+	check(IsInRenderingThread());
+
+	if (Capacity < RequestedCapacity)
+	{
+		// Apply some slack to limit number of rebuilds
+		uint32 NewCapacity = Slack > 0 ? RequestedCapacity * (1 + Slack) : RequestedCapacity;
+
+		if (GetDefault<ULidarPointCloudSettings>()->bLogBufferExpansion)
+		{
+			PC_LOG("Resizing %s: %u => %u", *GetFriendlyName(), Capacity, NewCapacity);
+		}
+
+		Release();
+		Capacity = NewCapacity;
+		Initialize();
+	}
+}
+
+//////////////////////////////////////////////////////////// Index Buffer
+
+FLidarPointCloudIndexBuffer::~FLidarPointCloudIndexBuffer()
+{
+	FRenderCommandFence Fence;
+
+	ENQUEUE_RENDER_COMMAND(ReleasePointCloudIndexBuffer)(
+		[this](FRHICommandListImmediate& RHICmdList)
+		{
+			Release();
+		});
+
+	Fence.BeginFence();
+	Fence.Wait();
+}
+
+void FLidarPointCloudIndexBuffer::InitRHI()
+{
+	FRHIResourceCreateInfo CreateInfo;
+	void* Buffer = nullptr;
+	uint32 Size = Capacity * 7 * sizeof(uint32);
+	PointOffset = Capacity * 6;
+
+	IndexBufferRHI = RHICreateAndLockIndexBuffer(sizeof(uint32), Size, BUF_Dynamic, CreateInfo, Buffer);
+
+	uint32* Data = (uint32*)Buffer;
+	for (uint32 i = 0; i < Capacity; i++)
+	{
+		// Full quads
+		{
+			uint32 idx = i * 6;
+			uint32 v = i * 4;
+
+			Data[idx] = v;
+			Data[idx + 1] = v + 1;
+			Data[idx + 2] = v + 2;
+			Data[idx + 3] = v;
+			Data[idx + 4] = v + 2;
+			Data[idx + 5] = v + 3;
+		}
+
+		// Points
+		Data[PointOffset + i] = i;
+	}
+
+	RHIUnlockIndexBuffer(IndexBufferRHI);
+	Buffer = nullptr;
+}
+
+//////////////////////////////////////////////////////////// Structured Buffer
+
+void FLidarPointCloudStructuredBuffer::Initialize()
+{
+	// This must be called from Rendering thread
+	check(IsInRenderingThread());
+
+	FRHIResourceCreateInfo CreateInfo;
+	Buffer = RHICreateStructuredBuffer(ElementSize, ElementSize * Capacity, BUF_ShaderResource | BUF_Dynamic, CreateInfo);
+	SRV = RHICreateShaderResourceView(Buffer);
+}
+
+void FLidarPointCloudStructuredBuffer::Release()
+{
+	// This must be called from Rendering thread
+	check(IsInRenderingThread());
+
+	if (Buffer)
+	{
+		RHIDiscardTransientResource(Buffer);
+		Buffer.SafeRelease();
+	}
+
+	SRV.SafeRelease();
+}
+
+//////////////////////////////////////////////////////////// Instance Buffer
+
+void FLidarPointCloudInstanceBuffer::Reset()
+{
+	NewNumInstances = 0;
+	Nodes.Reset();
+}
+
+uint32 FLidarPointCloudInstanceBuffer::UpdateData(bool bUseClassification)
+{
+	// Populate buffer only if there is any data to populate with
+	if (NewNumInstances)
+	{
+		// Rebuild resource if the new data size is greater than the current buffer can accommodate
+		// 17 bytes per point, element size set to 4 bytes - estimated 4.3 elements per point
+		StructuredBuffer.Resize(NewNumInstances * 4.3);
+
+		FColor SelectionColor = FColor::White;
+
+#if WITH_EDITOR
+		SelectionColor = GetDefault<UEditorStyleSettings>()->SelectionColor.ToFColor(false);
+#endif
+
+		// Send new data
+		uint8* BufferStart = (uint8*)RHILockStructuredBuffer(StructuredBuffer.Buffer, 0, StructuredBuffer.GetSize(), RLM_WriteOnly);
+		uint8* BufferCurrent = BufferStart;
+		const FLidarPointCloudPoint* DataPointer = nullptr;
+
+		if (bOwnedByEditor)
+		{
+			if (bUseClassification)
+			{
+				for (auto Node : Nodes)
+				{
+					const auto DataNode = Node->DataNode;
+					for (FLidarPointCloudPoint* Data = DataNode->AllocatedPoints.GetData(), *DataEnd = Data + DataNode->AllocatedPoints.Num(); Data != DataEnd; ++Data)
+					{
+						if (!Data->bVisible)
+						{
+							break;
+						}
+
+						FMemory::Memcpy(BufferCurrent, Data, 12);
+						BufferCurrent += 12;
+
+						if (Data->bSelected)
+						{
+							FMemory::Memcpy(BufferCurrent, &SelectionColor, 4);
+						}
+						else
+						{
+							FColor ClassificationColor(Data->ClassificationID, Data->ClassificationID, Data->ClassificationID, Data->Color.A);
+							FMemory::Memcpy(BufferCurrent, &ClassificationColor, 4);
+						}
+
+						BufferCurrent += 4;
+					}
+
+					for (FLidarPointCloudPoint* Data = DataNode->PaddingPoints.GetData(), *DataEnd = Data + DataNode->PaddingPoints.Num(); Data != DataEnd; ++Data)
+					{
+						if (!Data->bVisible)
+						{
+							break;
+						}
+
+						FMemory::Memcpy(BufferCurrent, Data, 12);
+						BufferCurrent += 12;
+
+						if (Data->bSelected)
+						{
+							FMemory::Memcpy(BufferCurrent, &SelectionColor, 4);
+						}
+						else
+						{
+							FColor ClassificationColor(Data->ClassificationID, Data->ClassificationID, Data->ClassificationID, Data->Color.A);
+							FMemory::Memcpy(BufferCurrent, &ClassificationColor, 4);
+						}
+
+						BufferCurrent += 4;
+					}
+				}
+			}
+			else
+			{
+				for (auto Node : Nodes)
+				{
+					const auto DataNode = Node->DataNode;
+					for (FLidarPointCloudPoint* Data = DataNode->AllocatedPoints.GetData(), *DataEnd = Data + DataNode->AllocatedPoints.Num(); Data != DataEnd; ++Data)
+					{
+						if (!Data->bVisible)
+						{
+							break;
+						}
+
+						if (Data->bSelected)
+						{
+							FMemory::Memcpy(BufferCurrent, Data, 12);
+							BufferCurrent += 12;
+
+							FMemory::Memcpy(BufferCurrent, &SelectionColor, 4);
+							BufferCurrent += 4;
+						}
+						else
+						{
+							FMemory::Memcpy(BufferCurrent, Data, 16);
+							BufferCurrent += 16;
+						}
+					}
+
+					for (FLidarPointCloudPoint* Data = DataNode->PaddingPoints.GetData(), *DataEnd = Data + DataNode->PaddingPoints.Num(); Data != DataEnd; ++Data)
+					{
+						if (!Data->bVisible)
+						{
+							break;
+						}
+
+						if (Data->bSelected)
+						{
+							FMemory::Memcpy(BufferCurrent, Data, 12);
+							BufferCurrent += 12;
+
+							FMemory::Memcpy(BufferCurrent, &SelectionColor, 4);
+							BufferCurrent += 4;
+						}
+						else
+						{
+							FMemory::Memcpy(BufferCurrent, Data, 16);
+							BufferCurrent += 16;
+						}
+					}
+				}
+			}
+		}
+		else
+		{
+			if (bUseClassification)
+			{
+				for (auto Node : Nodes)
+				{
+					const auto DataNode = Node->DataNode;
+					for (FLidarPointCloudPoint* Data = DataNode->AllocatedPoints.GetData(), *DataEnd = Data + DataNode->AllocatedPoints.Num(); Data != DataEnd; ++Data)
+					{
+						if (!Data->bVisible)
+						{
+							break;
+						}
+
+						FMemory::Memcpy(BufferCurrent, Data, 12);
+						BufferCurrent += 12;
+
+						FColor ClassificationColor(Data->ClassificationID, Data->ClassificationID, Data->ClassificationID, Data->Color.A);
+						FMemory::Memcpy(BufferCurrent, &ClassificationColor, 4);
+						BufferCurrent += 4;
+					}
+
+					for (FLidarPointCloudPoint* Data = DataNode->PaddingPoints.GetData(), *DataEnd = Data + DataNode->PaddingPoints.Num(); Data != DataEnd; ++Data)
+					{
+						if (!Data->bVisible)
+						{
+							break;
+						}
+
+						FMemory::Memcpy(BufferCurrent, Data, 12);
+						BufferCurrent += 12;
+
+						FColor ClassificationColor(Data->ClassificationID, Data->ClassificationID, Data->ClassificationID, Data->Color.A);
+						FMemory::Memcpy(BufferCurrent, &ClassificationColor, 4);
+						BufferCurrent += 4;
+					}
+				}
+			}
+			else
+			{
+				for (auto Node : Nodes)
+				{
+					const auto DataNode = Node->DataNode;
+					for (FLidarPointCloudPoint* Data = DataNode->AllocatedPoints.GetData(), *DataEnd = Data + DataNode->AllocatedPoints.Num(); Data != DataEnd; ++Data)
+					{
+						if (!Data->bVisible)
+						{
+							break;
+						}
+
+						FMemory::Memcpy(BufferCurrent, Data, 16);
+						BufferCurrent += 16;
+					}
+
+					for (FLidarPointCloudPoint* Data = DataNode->PaddingPoints.GetData(), *DataEnd = Data + DataNode->PaddingPoints.Num(); Data != DataEnd; ++Data)
+					{
+						if (!Data->bVisible)
+						{
+							break;
+						}
+
+						FMemory::Memcpy(BufferCurrent, Data, 16);
+						BufferCurrent += 16;
+					}
+				}
+			}
+		}
+
+		// Calculates the actual number of instances copied to the buffer (accounts for the invisible points)
+		NewNumInstances = (BufferCurrent - BufferStart) / 16;
+
+		// Send Scale data
+		for (auto Node : Nodes)
+		{
+			FMemory::Memset(BufferCurrent, Node->VirtualDepth, Node->DataNode->GetNumVisiblePoints());
+			BufferCurrent += Node->DataNode->GetNumVisiblePoints();
+		}
+		RHIUnlockStructuredBuffer(StructuredBuffer.Buffer);
+	}
+
+	return NumInstances = NewNumInstances;
+}
+
+void FLidarPointCloudInstanceBuffer::AddNode(const FLidarPointCloudTraversalOctreeNode* Node)
+{
+	Nodes.Add(Node);
+	NewNumInstances += Node->DataNode->GetNumVisiblePoints();
+}
+
+void FLidarPointCloudInstanceBuffer::Release()
+{
+	Nodes.Empty();
+	NumInstances = 0;
+}
+
+//////////////////////////////////////////////////////////// User Data
+
+void FLidarPointCloudBatchElementUserData::SetClassificationColors(const TMap<int32, FLinearColor>& InClassificationColors)
+{
+	for (int32 i = 0; i < 32; ++i)
+	{
+		const FLinearColor* Color = InClassificationColors.Find(i);
+		ClassificationColors[i] = Color ? FVector4(*Color) : FVector4(1, 1, 1);
+	}
+}
+
+//////////////////////////////////////////////////////////// Vertex Factory
+
+void FLidarPointCloudVertexFactoryShaderParameters::Bind(const FShaderParameterMap& ParameterMap)
+{
+	BINDPARAM(DataBuffer);
+	BINDPARAM(IndexDivisor);
+	BINDPARAM(VDMultiplier);
+	BINDPARAM(SizeOffset);
+	BINDPARAM(RootCellSize);
+	BINDPARAM(bUseLODColoration);
+	BINDPARAM(SpriteSizeMultiplier);
+	BINDPARAM(ViewRightVector);
+	BINDPARAM(ViewUpVector);
+	BINDPARAM(BoundsSize);
+	BINDPARAM(BoundsOffset);
+	BINDPARAM(ElevationColorBottom);
+	BINDPARAM(ElevationColorTop);
+	BINDPARAM(bUseCircle);
+	BINDPARAM(bUseColorOverride);
+	BINDPARAM(bUseElevationColor);
+	BINDPARAM(Offset);
+	BINDPARAM(Contrast);
+	BINDPARAM(Saturation);
+	BINDPARAM(Gamma);
+	BINDPARAM(Tint);
+	BINDPARAM(IntensityInfluence);
+	BINDPARAM(bUseClassification);
+	BINDPARAM(ClassificationColors);
+}
+
+void FLidarPointCloudVertexFactoryShaderParameters::GetElementShaderBindings(const class FSceneInterface* Scene, const FSceneView* View, const FMeshMaterialShader* Shader, const EVertexInputStreamType InputStreamType, ERHIFeatureLevel::Type FeatureLevel,
+	const FVertexFactory* VertexFactory, const FMeshBatchElement& BatchElement, class FMeshDrawSingleShaderBindings& ShaderBindings, FVertexInputStreamArray& VertexStreams) const
+{
+	FLidarPointCloudBatchElementUserData* UserData = (FLidarPointCloudBatchElementUserData*)BatchElement.UserData;
+
+	SETSRVPARAM(DataBuffer);
+	SETPARAM(IndexDivisor);
+	SETPARAM(VDMultiplier);
+	SETPARAM(SizeOffset);
+	SETPARAM(RootCellSize);
+	SETPARAM(bUseLODColoration);
+	SETPARAM(SpriteSizeMultiplier);
+	SETPARAM(ViewRightVector);
+	SETPARAM(ViewUpVector);
+	SETPARAM(BoundsSize);
+	SETPARAM(BoundsOffset);
+	SETPARAM(ElevationColorBottom);
+	SETPARAM(ElevationColorTop);
+	SETPARAM(bUseCircle);
+	SETPARAM(bUseColorOverride);
+	SETPARAM(bUseElevationColor);
+	SETPARAM(Offset);
+	SETPARAM(Contrast);
+	SETPARAM(Saturation);
+	SETPARAM(Gamma);
+	SETPARAM(Tint);
+	SETPARAM(IntensityInfluence);
+	SETPARAM(bUseClassification);
+	SETPARAM(ClassificationColors);
+}
+
+FVertexFactoryShaderParameters* FLidarPointCloudVertexFactory::ConstructShaderParameters(EShaderFrequency ShaderFrequency)
+{
+	if (ShaderFrequency == SF_Vertex)
+	{
+		return new FLidarPointCloudVertexFactoryShaderParameters();
+	}
+
+	return nullptr;
+}
+
+FLidarPointCloudVertexFactory::~FLidarPointCloudVertexFactory()
+{
+	FRenderCommandFence Fence;
+
+	FVertexFactory* VertexFactory = this;
+
+	ENQUEUE_RENDER_COMMAND(ReleasePointCloudIndexBuffer)(
+		[VertexFactory](FRHICommandListImmediate& RHICmdList)
+		{
+			VertexFactory->ReleaseResource();
+		});
+
+	Fence.BeginFence();
+	Fence.Wait();
+}
+
+bool FLidarPointCloudVertexFactory::ShouldCompilePermutation(EShaderPlatform Platform, const class FMaterial* Material, const class FShaderType* ShaderType)
+{
+	return (IsPCPlatform(Platform) && IsFeatureLevelSupported(Platform, ERHIFeatureLevel::SM5) && Material->GetMaterialDomain() == MD_Surface && Material->IsUsedWithLidarPointCloud()) || Material->IsSpecialEngineMaterial();
+}
+
+void FLidarPointCloudVertexFactory::InitRHI()
+{
+	VertexBuffer.InitResource();
+
+	FVertexDeclarationElementList Elements;
+	Elements.Add(AccessStreamComponent(FVertexStreamComponent(&VertexBuffer, 0, 0, VET_Float3), 0));
+	InitDeclaration(Elements);
+}
+
+void FLidarPointCloudVertexFactory::ReleaseRHI()
+{
+	FVertexFactory::ReleaseRHI();
+	VertexBuffer.ReleaseResource();
+}
+
+IMPLEMENT_VERTEX_FACTORY_TYPE(FLidarPointCloudVertexFactory, "/Plugin/LidarPointCloud/Private/LidarPointCloudVertexFactory.ush", /* bUsedWithMaterials */ true, /* bSupportsStaticLighting */ false, /* bSupportsDynamicLighting */ true, /* bPrecisePrevWorldPos */ false, /* bSupportsPositionOnly */ true);
+
+#undef BINDPARAM
+#undef SETPARAM
+#undef SETSRVPARAM
diff --git a/Engine/Plugins/Enterprise/LidarPointCloud/Source/LidarPointCloudRuntime/Private/Rendering/LidarPointCloudRenderBuffers.h b/Engine/Plugins/Enterprise/LidarPointCloud/Source/LidarPointCloudRuntime/Private/Rendering/LidarPointCloudRenderBuffers.h
new file mode 100644
--- /dev/null
+++ b/Engine/Plugins/Enterprise/LidarPointCloud/Source/LidarPointCloudRuntime/Private/Rendering/LidarPointCloudRenderBuffers.h
@@ -0,0 +1,247 @@
+// Copyright Epic Games, Inc. All Rights Reserved.
+
+#pragma once
+
+#include "CoreMinimal.h"
+#include "RenderResource.h"
+#include "VertexFactory.h"
+#include "LidarPointCloudShared.h"
+#include "HAL/ThreadSafeBool.h"
+
+/**
+ * Base class for the buffers
+ */
+class FLidarPointCloudBuffer
+{
+public:
+	/** Resizes the buffer to the specified capacity, if necessary. Must be called from Rendering thread. */
+	void Resize(const uint32& RequestedCapacity);
+
+	FLidarPointCloudBuffer(const uint32& ElementSize, const float& Slack, const uint32 InitialCapacity = 0) : Capacity(InitialCapacity), ElementSize(ElementSize), Slack(Slack) {}
+
+	virtual FString GetFriendlyName() const { return TEXT("FPointCloudBuffer"); }
+
+	/** Returns the total size of this buffer */
+	virtual uint32 GetSize() const { return Capacity * ElementSize; }
+
+protected:
+	uint32 Capacity;
+	uint32 ElementSize;
+	float Slack;
+
+	FLidarPointCloudBuffer() : FLidarPointCloudBuffer(0, 0) {}
+
+	virtual void Initialize() = 0;
+	virtual void Release() = 0;
+};
+
+/**
+ * This class creates an IndexBuffer shared between all assets and all instances.
+ */
+class FLidarPointCloudIndexBuffer : public FIndexBuffer, public FLidarPointCloudBuffer
+{
+public:
+	FLidarPointCloudIndexBuffer() : FLidarPointCloudBuffer(sizeof(uint32), 0, 1000000) { }
+	~FLidarPointCloudIndexBuffer();
+
+	uint32 GetPointOffset() const { return PointOffset; }
+
+private:
+	virtual void InitRHI() override;
+	virtual void Initialize() override { InitResource(); }
+	virtual void Release() override { ReleaseResource(); }
+	virtual FString GetFriendlyName() const override { return TEXT("FLidarPointCloudIndexBuffer"); }
+
+private:
+	uint32 PointOffset;
+};
+
+/**
+ * Encapsulates a GPU read structured buffer with its SRV.
+ */
+class FLidarPointCloudStructuredBuffer : public FLidarPointCloudBuffer
+{
+public:
+	FStructuredBufferRHIRef Buffer;
+	FShaderResourceViewRHIRef SRV;
+
+	FLidarPointCloudStructuredBuffer(const uint32& ElementSize, const float& Slack) : FLidarPointCloudBuffer(ElementSize, Slack) {}
+	virtual ~FLidarPointCloudStructuredBuffer() { Release(); }
+
+	virtual void Initialize() override;
+	virtual void Release() override;
+	virtual FString GetFriendlyName() const override { return TEXT("Structured Buffer"); }
+
+private:
+	FLidarPointCloudStructuredBuffer(const FLidarPointCloudStructuredBuffer&) = delete;
+	FLidarPointCloudStructuredBuffer(FLidarPointCloudStructuredBuffer&&) = delete;
+	FLidarPointCloudStructuredBuffer& operator=(const FLidarPointCloudStructuredBuffer&) = delete;
+	FLidarPointCloudStructuredBuffer& operator=(FLidarPointCloudStructuredBuffer&&) = delete;
+};
+
+/**
+ * Handles creation of the dynamic Instance Data
+ */
+// #refactor: Can we avoid referencing FLidarPointCloudOctreeNode here?
+struct FLidarPointCloudTraversalOctreeNode;
+class FLidarPointCloudInstanceBuffer
+{
+public:
+	FLidarPointCloudStructuredBuffer StructuredBuffer;
+	TArray<const FLidarPointCloudTraversalOctreeNode*> Nodes;
+	bool bOwnedByEditor;
+
+private:
+	uint32 NumInstances;
+	uint32 NewNumInstances;
+
+public:
+	FLidarPointCloudInstanceBuffer()
+		: StructuredBuffer(4, 1.0f)
+		, bOwnedByEditor(false)
+		, NumInstances(0)
+		, NewNumInstances(0)
+	{
+	}
+
+	FORCEINLINE uint32 GetNumInstances() const { return NumInstances; }
+
+	/** Resets the nodes assigned for building this buffer */
+	void Reset();
+
+	/**
+	 * Resizes and updates the buffer on the GPU if necessary.
+	 * Returns the number of instances used.
+	 */
+	uint32 UpdateData(bool bUseClassification);
+
+	void AddNode(const FLidarPointCloudTraversalOctreeNode* Node);
+
+	/** Releases the memory held by StructuredBuffer */
+	void Release();
+};
+
+/**
+ * Holds all data to be passed to the FLidarPointCloudVertexFactoryShaderParameters as UserData
+ */
+struct FLidarPointCloudBatchElementUserData
+{
+	FRHIShaderResourceView* DataBuffer;
+	int32 IndexDivisor;
+	float VDMultiplier;
+	int32 SizeOffset;
+	float RootCellSize;
+	int32 bUseLODColoration;
+	float SpriteSizeMultiplier;
+	FVector ViewRightVector;
+	FVector ViewUpVector;
+	FVector BoundsSize;
+	FVector BoundsOffset;
+	FVector ElevationColorBottom;
+	FVector ElevationColorTop;
+	int32 bUseCircle;
+	int32 bUseColorOverride;
+	int32 bUseElevationColor;
+	FVector4 Offset;
+	FVector4 Contrast;
+	FVector4 Saturation;
+	FVector4 Gamma;
+	FVector Tint;
+	float IntensityInfluence;
+	int32 bUseClassification;
+	FVector4 ClassificationColors[32];
+
+	FLidarPointCloudBatchElementUserData(const float& VDMultiplier, const float& RootCellSize)
+		: DataBuffer(nullptr)
+		, IndexDivisor(4)
+		, VDMultiplier(VDMultiplier)
+		, RootCellSize(RootCellSize)
+		, bUseLODColoration(false)
+	{
+	}
+
+	void SetClassificationColors(const TMap<int32, FLinearColor>& InClassificationColors);
+};
+
+/**
+ * Binds shader parameters necessary for rendering
+ */
+class FLidarPointCloudVertexFactoryShaderParameters : public FVertexFactoryShaderParameters
+{
+public:
+	virtual void Bind(const FShaderParameterMap& ParameterMap) override;
+	virtual void Serialize(FArchive& Ar) override
+	{
+		Ar << DataBuffer << IndexDivisor << VDMultiplier << SizeOffset << RootCellSize << bUseLODColoration << SpriteSizeMultiplier << ViewRightVector << ViewUpVector
+			<< BoundsSize << BoundsOffset << ElevationColorBottom << ElevationColorTop << bUseCircle << bUseColorOverride << bUseElevationColor
+			<< Offset << Contrast << Saturation << Gamma << Tint << IntensityInfluence << bUseClassification << ClassificationColors;
+	}
+
+	virtual void GetElementShaderBindings(const class FSceneInterface* Scene, const FSceneView* View, const FMeshMaterialShader* Shader, const EVertexInputStreamType InputStreamType, ERHIFeatureLevel::Type FeatureLevel,
+		const FVertexFactory* VertexFactory, const FMeshBatchElement& BatchElement, class FMeshDrawSingleShaderBindings& ShaderBindings, FVertexInputStreamArray& VertexStreams) const override;
+
+private:
+	FShaderResourceParameter DataBuffer;
+	FShaderParameter IndexDivisor;
+	FShaderParameter VDMultiplier;
+	FShaderParameter SizeOffset;
+	FShaderParameter RootCellSize;
+	FShaderParameter bUseLODColoration;
+	FShaderParameter SpriteSizeMultiplier;
+	FShaderParameter ViewRightVector;
+	FShaderParameter ViewUpVector;
+	FShaderParameter BoundsSize;
+	FShaderParameter BoundsOffset;
+	FShaderParameter ElevationColorBottom;
+	FShaderParameter ElevationColorTop;
+	FShaderParameter bUseCircle;
+	FShaderParameter bUseColorOverride;
+	FShaderParameter bUseElevationColor;
+	FShaderParameter Offset;
+	FShaderParameter Contrast;
+	FShaderParameter Saturation;
+	FShaderParameter Gamma;
+	FShaderParameter Tint;
+	FShaderParameter IntensityInfluence;
+	FShaderParameter bUseClassification;
+	FShaderParameter ClassificationColors;
+};
+
+/**
+ * Implementation of the custom Vertex Factory, containing only a ZeroStride position stream.
+ */
+class FLidarPointCloudVertexFactory : public FVertexFactory
+{
+	DECLARE_VERTEX_FACTORY_TYPE(FLidarPointCloudVertexFactory);
+
+public:
+	static FVertexFactoryShaderParameters* ConstructShaderParameters(EShaderFrequency ShaderFrequency);
+
+	static bool ShouldCache(EShaderPlatform Platform, const class FMaterial* Material, const class FShaderType* ShaderType) { return IsFeatureLevelSupported(Platform, ERHIFeatureLevel::SM5); }
+	static bool ShouldCompilePermutation(EShaderPlatform Platform, const class FMaterial* Material, const class FShaderType* ShaderType);
+	static void ModifyCompilationEnvironment(const FVertexFactoryType* Type, EShaderPlatform Platform, const FMaterial* Material, FShaderCompilerEnvironment& OutEnvironment) {}
+
+	FLidarPointCloudVertexFactory() : FVertexFactory(ERHIFeatureLevel::SM5) { }
+	~FLidarPointCloudVertexFactory();
+
+private:
+	/** Very simple implementation of a ZeroStride Vertex Buffer */
+	class FPointCloudVertexBuffer : public FVertexBuffer
+	{
+	public:
+		virtual void InitRHI() override
+		{
+			FRHIResourceCreateInfo CreateInfo;
+			void* Buffer = nullptr;
+			VertexBufferRHI = RHICreateAndLockVertexBuffer(sizeof(FVector), BUF_Static | BUF_ZeroStride, CreateInfo, Buffer);
+			FMemory::Memzero(Buffer, sizeof(FVector));
+			RHIUnlockVertexBuffer(VertexBufferRHI);
+			Buffer = nullptr;
+		}
+
+		virtual FString GetFriendlyName() const override { return TEXT("FPointCloudVertexBuffer"); }
+	} VertexBuffer;
+
+	virtual void InitRHI() override;
+	virtual void ReleaseRHI() override;
+};
\ No newline at end of file
