Copy up Dev-Enterprise @ CL11133089

diff --git a/Engine/Plugins/Enterprise/LidarPointCloud/Source/LidarPointCloudRuntime/Private/Rendering/LidarPointCloudRendering.cpp b/Engine/Plugins/Enterprise/LidarPointCloud/Source/LidarPointCloudRuntime/Private/Rendering/LidarPointCloudRendering.cpp
new file mode 100644
--- /dev/null
+++ b/Engine/Plugins/Enterprise/LidarPointCloud/Source/LidarPointCloudRuntime/Private/Rendering/LidarPointCloudRendering.cpp
@@ -0,0 +1,509 @@
+// Copyright Epic Games, Inc. All Rights Reserved.
+
+#include "LidarPointCloudRenderBuffers.h"
+#include "LidarPointCloudComponent.h"
+#include "LidarPointCloud.h"
+#include "LidarPointCloudShared.h"
+#include "LidarPointCloudOctree.h"
+#include "PrimitiveSceneProxy.h"
+#include "MeshBatch.h"
+#include "Engine/Engine.h"
+#include "Misc/ScopeTryLock.h"
+#include "SceneManagement.h"
+#include "LocalVertexFactory.h"
+#include "Materials/Material.h"
+
+DECLARE_DWORD_COUNTER_STAT(TEXT("Visible Points"), STAT_PointCount, STATGROUP_LidarPointCloud)
+DECLARE_DWORD_COUNTER_STAT(TEXT("Visible Nodes"), STAT_NodeCount, STATGROUP_LidarPointCloud)
+DECLARE_DWORD_COUNTER_STAT(TEXT("Draw Calls"), STAT_DrawCallCount, STATGROUP_LidarPointCloud)
+
+DECLARE_CYCLE_STAT(TEXT("Buffer Update"), STAT_BufferUpdate, STATGROUP_LidarPointCloud);
+
+/** A global instance of an Index Buffer shared between all Lidar Point Cloud proxies */
+TGlobalResource<FLidarPointCloudIndexBuffer> GLidarPointCloudIndexBuffer;
+
+/** A global instance of a Vertex Factory shared between all Lidar Point Cloud proxies */
+TGlobalResource<FLidarPointCloudVertexFactory> GLidarPointCloudVertexFactory;
+
+class FLidarPointCloudCollisionRendering
+{
+public:
+	FLidarPointCloudCollisionRendering()
+		: NumPrimitives(0)
+		, MaxVertexIndex(0)
+	{
+	}
+	~FLidarPointCloudCollisionRendering()
+	{
+		VertexFactory.ReleaseResource();
+		VertexBuffer.ReleaseResource();
+		IndexBuffer.ReleaseResource();
+	}
+
+	void Initialize(FLidarPointCloudOctree* Octree)
+	{
+		// Create collision visualization buffers
+		if (Octree->HasCollisionData())
+		{
+			auto CollisionData = Octree->GetCollisionData();
+			
+			// Release the existing resources, if present
+			VertexBuffer.ReleaseResource();
+			VertexFactory.ReleaseResource();
+			IndexBuffer.ReleaseResource();
+
+			// Initialize the buffers
+			VertexBuffer.Initialize(CollisionData->Vertices.GetData(), CollisionData->Vertices.Num());
+			VertexFactory.Initialize(&VertexBuffer);
+			IndexBuffer.Initialize((int32*)CollisionData->Indices.GetData(), CollisionData->Indices.Num() * 3);
+
+			NumPrimitives = CollisionData->Indices.Num();
+			MaxVertexIndex = CollisionData->Vertices.Num() - 1;
+		}
+	}
+
+	const FVertexFactory* GetVertexFactory() const { return &VertexFactory; }
+	const FIndexBuffer* GetIndexBuffer() const { return &IndexBuffer; }
+	const int32 GetNumPrimitives() const { return NumPrimitives; }
+	const int32 GetMaxVertexIndex() const { return MaxVertexIndex; }
+
+private:
+	class FLidarPointCloudCollisionVertexFactory : public FLocalVertexFactory
+	{
+	public:
+		FLidarPointCloudCollisionVertexFactory() : FLocalVertexFactory(ERHIFeatureLevel::SM5, "") { }
+
+		void Initialize(FVertexBuffer* InVertexBuffer)
+		{
+			FDataType NewData;
+			NewData.PositionComponent = FVertexStreamComponent(InVertexBuffer, 0, 12, VET_Float3);
+			NewData.TangentBasisComponents[0] = FVertexStreamComponent(InVertexBuffer, 0, 12, VET_PackedNormal);
+			NewData.TangentBasisComponents[1] = FVertexStreamComponent(InVertexBuffer, 0, 12, VET_PackedNormal);
+			NewData.ColorComponent = FVertexStreamComponent(InVertexBuffer, 0, 12, VET_Color);
+
+			NewData.PositionComponentSRV = nullptr;
+			NewData.ColorComponentsSRV = nullptr;
+			NewData.TangentsSRV = nullptr;
+			NewData.TextureCoordinatesSRV = nullptr;
+
+			if (RHISupportsManualVertexFetch(GMaxRHIShaderPlatform))
+			{
+				NewData.PositionComponentSRV = RHICreateShaderResourceView(InVertexBuffer->VertexBufferRHI, 4, PF_R32_FLOAT);
+				NewData.ColorComponentsSRV = RHICreateShaderResourceView(nullptr, 4, PF_R8G8B8A8);
+				NewData.TangentsSRV = RHICreateShaderResourceView(nullptr, 4, PF_R8G8B8A8);
+				NewData.TextureCoordinatesSRV = RHICreateShaderResourceView(nullptr, 4, PF_G16R16F);
+			}
+
+			Data = NewData;
+			InitResource();
+		}
+	} VertexFactory;
+	class FLidarPointCloudCollisionVertexBuffer : public FVertexBuffer
+	{
+	private:
+		const void* Data;
+		int32 DataLength;
+
+	public:
+		void Initialize(const FVector* InData, int32 InDataLength)
+		{
+			Data = InData;
+			DataLength = InDataLength;
+
+			InitResource();
+		}
+
+		virtual void InitRHI() override
+		{
+			FRHIResourceCreateInfo CreateInfo;
+			void* Buffer = nullptr;
+			VertexBufferRHI = RHICreateAndLockVertexBuffer(DataLength * sizeof(FVector), BUF_Static | BUF_ShaderResource, CreateInfo, Buffer);
+			FMemory::Memcpy(Buffer, Data, DataLength * sizeof(FVector));
+			RHIUnlockVertexBuffer(VertexBufferRHI);
+			Buffer = nullptr;
+		}
+	} VertexBuffer;
+	class FLidarPointCloudCollisionIndexBuffer : public FIndexBuffer
+	{
+	private:
+		const int32* Data;
+		int32 DataLength;
+
+	public:
+		void Initialize(const int32* InData, int32 InDataLength)
+		{
+			Data = InData;
+			DataLength = InDataLength;
+
+			InitResource();
+		}
+
+		virtual void InitRHI() override
+		{
+			FRHIResourceCreateInfo CreateInfo;
+			void* Buffer = nullptr;
+			IndexBufferRHI = RHICreateAndLockIndexBuffer(sizeof(uint32), DataLength * sizeof(uint32), BUF_Static, CreateInfo, Buffer);
+			FMemory::Memcpy(Buffer, Data, DataLength * sizeof(uint32));
+			RHIUnlockIndexBuffer(IndexBufferRHI);
+			Buffer = nullptr;
+		}
+	} IndexBuffer;
+
+	int32 NumPrimitives;
+	int32 MaxVertexIndex;
+};
+
+class FLidarPointCloudSceneProxy : public FPrimitiveSceneProxy
+{
+public:
+	FLidarPointCloudSceneProxy(ULidarPointCloudComponent* Component)
+		: FPrimitiveSceneProxy(Component)
+		, bReadyToRender(false)
+		, Component(Component)
+		, Octree(&Component->GetPointCloud()->Octree)
+		, Owner(Component->GetOwner())
+		, CollisionRendering(Component->GetPointCloud()->CollisionRendering)
+	{
+		// Skip material verification - async update could occasionally cause it to crash
+		bVerifyUsedMaterials = false;
+
+		MaterialRelevance = Component->GetMaterialRelevance(GetScene().GetFeatureLevel());
+	}
+
+	virtual void GetDynamicMeshElements(const TArray<const FSceneView*>& Views, const FSceneViewFamily& ViewFamily, uint32 VisibilityMap, FMeshElementCollector& Collector) const override
+	{
+		QUICK_SCOPE_CYCLE_COUNTER(STAT_PointCloudSceneProxy_GetDynamicMeshElements);
+
+		if (!bReadyToRender)
+		{
+			return;
+		}
+
+		const bool bUsesSprites = Component->PointSize > 0;
+
+		for (int32 ViewIndex = 0; ViewIndex < Views.Num(); ViewIndex++)
+		{
+			const FSceneView* View = Views[ViewIndex];
+
+			if (IsShown(View) && (VisibilityMap & (1 << ViewIndex)))
+			{
+				// Prepare the draw call
+				int32 NumElements = InstanceBuffer.GetNumInstances();
+				if (NumElements)
+				{
+					FMeshBatch& MeshBatch = Collector.AllocateMesh();
+
+					MeshBatch.Type = bUsesSprites ? PT_TriangleList : PT_PointList;
+					MeshBatch.LODIndex = 0;
+					MeshBatch.VertexFactory = &GLidarPointCloudVertexFactory;
+					MeshBatch.bWireframe = false;
+					MeshBatch.MaterialRenderProxy = Component->GetMaterial(0)->GetRenderProxy();
+					MeshBatch.ReverseCulling = IsLocalToWorldDeterminantNegative();
+					MeshBatch.DepthPriorityGroup = SDPG_World;
+
+					FMeshBatchElement& BatchElement = MeshBatch.Elements[0];
+					BatchElement.PrimitiveUniformBuffer = GetUniformBuffer();
+					BatchElement.IndexBuffer = &GLidarPointCloudIndexBuffer;
+					BatchElement.FirstIndex = bUsesSprites ? 0 : GLidarPointCloudIndexBuffer.GetPointOffset();
+					BatchElement.MinVertexIndex = 0;
+					BatchElement.NumPrimitives = NumElements * (bUsesSprites ? 2 : 1);
+					BatchElement.UserData = &UserData[UserData.Add(BuildUserDataElement(View))];
+
+					Collector.AddMesh(ViewIndex, MeshBatch);
+
+					INC_DWORD_STAT(STAT_DrawCallCount);
+				}
+
+#if !(UE_BUILD_SHIPPING)
+				FPrimitiveDrawInterface* PDI = Collector.GetPDI(ViewIndex);
+
+				// Draw selected nodes' bounds
+				if (Component->bDrawNodeBounds)
+				{
+					for (auto Node : SelectedNodes)
+					{
+						FVector Extent = TraversalOctree->Extents[Node->Depth];
+						DrawWireBox(PDI, FBox(Node->Center - Extent, Node->Center + Extent), FColor(72, 72, 255), SDPG_World);
+					}
+				}
+
+				// Render bounds
+				if (ViewFamily.EngineShowFlags.Bounds)
+				{
+					RenderBounds(PDI, ViewFamily.EngineShowFlags, GetBounds(), !Owner || IsSelected());
+				}
+
+				// Render collision wireframe
+				if (ViewFamily.EngineShowFlags.Collision && IsCollisionEnabled() && CollisionRendering && CollisionRendering->GetNumPrimitives() > 0)
+				{
+					// Create colored proxy
+					FColoredMaterialRenderProxy* CollisionMaterialInstance;
+					CollisionMaterialInstance = new FColoredMaterialRenderProxy(GEngine->WireframeMaterial->GetRenderProxy(), FColor(0, 255, 255, 255));
+					Collector.RegisterOneFrameMaterialProxy(CollisionMaterialInstance);
+
+					FMeshBatch& MeshBatch = Collector.AllocateMesh();
+					if (PrepareCollisionWireframe(MeshBatch, CollisionMaterialInstance))
+					{
+						Collector.AddMesh(ViewIndex, MeshBatch);
+					}
+				}
+#endif // !(UE_BUILD_SHIPPING)
+			}
+		}
+	}
+
+	virtual FPrimitiveViewRelevance GetViewRelevance(const FSceneView* View) const override
+	{
+		FPrimitiveViewRelevance Result;
+
+		Result.bUseCustomViewData = true;
+		Result.bDrawRelevance = IsShown(View);
+		Result.bShadowRelevance = IsShadowCast(View);
+		Result.bDynamicRelevance = true;
+		Result.bStaticRelevance = false;
+		Result.bRenderInMainPass = ShouldRenderInMainPass();
+		Result.bUsesLightingChannels = GetLightingChannelMask() != GetDefaultLightingChannelMask();
+		Result.bRenderCustomDepth = ShouldRenderCustomDepth();
+		MaterialRelevance.SetPrimitiveViewRelevance(Result);
+
+		return Result;
+	}
+
+	/** Calculates node visibility and updates InstanceBuffer */
+	virtual void* InitViewCustomData(const FSceneView& InView, float InViewLODScale, FMemStackBase& InCustomDataMemStack, bool InIsStaticRelevant, bool InIsShadowOnly, const FLODMask* InVisiblePrimitiveLODMask = nullptr, float InMeshScreenSizeSquared = -1.0f) override
+	{
+		// Don't process for shadow views
+		if (InIsShadowOnly)
+		{
+			return nullptr;
+		}
+
+		bReadyToRender = false;
+
+		// Check if the data is available
+		if (!Octree || Octree->GetNumNodes() == 0 || !TraversalOctree.IsValid() || TraversalOctree->NumLODs == 0 || !Component->GetMaterial(0))
+		{
+			return nullptr;
+		}
+
+		// Check if the access to the data has been restricted
+		FScopeTryLock Lock(&Octree->DataLock);
+		if (!Lock.IsLocked())
+		{
+			return nullptr;
+		}
+
+		// Re-initialize the traversal octree, if needed
+		if (!TraversalOctree->bValid)
+		{
+			TraversalOctree.Reset(new FLidarPointCloudTraversalOctree(Octree, Component->GetComponentTransform()));
+		}
+
+		// Clamp the budget to 25M
+		const int32 PointBudget = FMath::Clamp(Component->PointBudget, 0, 25000000);
+
+		// Expand the Index Buffer, if necessary
+		GLidarPointCloudIndexBuffer.Resize(PointBudget);
+
+		// Prepare the buffer
+		InstanceBuffer.Reset();
+
+		// Select the relevant nodes to be rendered
+		SelectedNodes = TraversalOctree->GetVisibleNodes(&InView, FMath::Max(Component->MinScreenSize * Component->MinScreenSize, 0.0f), PointBudget, Component->ScreenCenterImportance, Component->MinDepth, Component->MaxDepth);
+		INC_DWORD_STAT_BY(STAT_NodeCount, SelectedNodes.Num());
+
+		// Calculate sprite size, if needed
+		if (Component->PointSize > 0)
+		{
+			TraversalOctree->CalculateSpriteSize(SelectedNodes, Component->PointSizeBias);
+		}
+
+		// Allocate the selected nodes
+		for (auto Node : SelectedNodes)
+		{
+			InstanceBuffer.AddNode(Node);
+		}
+
+		// Update buffers
+		{
+			SCOPE_CYCLE_COUNTER(STAT_BufferUpdate);
+			InstanceBuffer.UpdateData(Component->ColorSource == ELidarPointCloudColorationMode::Classification);
+		}
+
+		bReadyToRender = true;
+
+		UserData.Reset(10);
+
+		INC_DWORD_STAT_BY(STAT_PointCount, InstanceBuffer.GetNumInstances());
+
+		return nullptr;
+	}
+
+	bool PrepareCollisionWireframe(FMeshBatch& MeshBatch, FColoredMaterialRenderProxy* CollisionMaterialInstance) const
+	{
+		MeshBatch.Type = PT_TriangleList;
+		MeshBatch.VertexFactory = CollisionRendering->GetVertexFactory();
+		MeshBatch.bWireframe = true;
+		MeshBatch.MaterialRenderProxy = CollisionMaterialInstance;
+		MeshBatch.ReverseCulling = !IsLocalToWorldDeterminantNegative();
+		MeshBatch.DepthPriorityGroup = SDPG_World;
+		MeshBatch.CastShadow = false;
+
+		FMeshBatchElement& BatchElement = MeshBatch.Elements[0];
+		BatchElement.IndexBuffer = CollisionRendering->GetIndexBuffer();
+		BatchElement.FirstIndex = 0;
+		BatchElement.NumPrimitives = CollisionRendering->GetNumPrimitives();
+		BatchElement.MinVertexIndex = 0;
+		BatchElement.MaxVertexIndex = CollisionRendering->GetMaxVertexIndex();
+
+		return true;
+	}
+
+	/** UserData is used to pass rendering information to the VertexFactory */
+	FLidarPointCloudBatchElementUserData BuildUserDataElement(const FSceneView* InView) const
+	{
+		// Each point has a virtual LOD in range of 0 to Octree->GetNumLODs(), mapped to 0 - 255 range
+		// To reverse the mapping: x / 255.0f * Octree->GetNumLODs() or x * (1.0f / 255.0f * Octree->GetNumLODs())
+		constexpr float InvertedMultiplier = 1.0f / 255.0f;
+		FLidarPointCloudBatchElementUserData UserDataElement(InvertedMultiplier * Octree->GetNumLODs(), Octree->GetRootCellSize());
+
+		const bool bUsesSprites = Component->PointSize > 0;
+
+		auto BoundsSize = Component->GetPointCloud()->GetPointsBounds().GetSize();
+
+		// Make sure to apply minimum bounds size
+		BoundsSize.X = FMath::Max(BoundsSize.X, 0.001f);
+		BoundsSize.Y = FMath::Max(BoundsSize.Y, 0.001f);
+		BoundsSize.Z = FMath::Max(BoundsSize.Z, 0.001f);
+
+		// Update shader parameters
+		UserDataElement.IndexDivisor = bUsesSprites ? 4 : 1;
+
+		UserDataElement.SizeOffset = InstanceBuffer.GetNumInstances() * 4;
+		UserDataElement.SpriteSizeMultiplier = Component->PointSize * Component->GetComponentScale().GetAbsMax();
+
+		UserDataElement.ViewRightVector = InView->GetViewRight();
+		UserDataElement.ViewUpVector = InView->GetViewUp();
+		UserDataElement.BoundsSize = BoundsSize;
+		UserDataElement.BoundsOffset = Component->GetPointCloud()->GetBounds().GetCenter() - Component->GetPointCloud()->GetPointsBounds().GetCenter();
+		UserDataElement.ElevationColorBottom = FVector(Component->ColorSource == ELidarPointCloudColorationMode::None ? FColor::White : Component->ElevationColorBottom);
+		UserDataElement.ElevationColorTop = FVector(Component->ColorSource == ELidarPointCloudColorationMode::None ? FColor::White : Component->ElevationColorTop);
+		UserDataElement.bUseCircle = bUsesSprites && Component->PointShape == ELidarPointCloudSpriteShape::Circle;
+		UserDataElement.bUseColorOverride = Component->ColorSource != ELidarPointCloudColorationMode::Data;
+		UserDataElement.bUseElevationColor = Component->ColorSource == ELidarPointCloudColorationMode::Elevation || Component->ColorSource == ELidarPointCloudColorationMode::None;
+		UserDataElement.Offset = Component->Offset;
+		UserDataElement.Contrast = Component->Contrast;
+		UserDataElement.Saturation = Component->Saturation;
+		UserDataElement.Gamma = Component->Gamma;
+		UserDataElement.Tint = FVector(Component->ColorTint);
+		UserDataElement.IntensityInfluence = Component->IntensityInfluence;
+
+		UserDataElement.bUseClassification = Component->ColorSource == ELidarPointCloudColorationMode::Classification;
+		UserDataElement.SetClassificationColors(Component->ClassificationColors);
+
+		// Make sure to update the reference to the resource, in case it was re-initialized
+		UserDataElement.DataBuffer = InstanceBuffer.StructuredBuffer.SRV;
+
+		return UserDataElement;
+	}
+
+	virtual bool CanBeOccluded() const override { return !MaterialRelevance.bDisableDepthTest; }
+	virtual uint32 GetMemoryFootprint() const override { return(sizeof(*this) + GetAllocatedSize()); }
+	uint32 GetAllocatedSize() const { return(FPrimitiveSceneProxy::GetAllocatedSize()); }
+
+	virtual SIZE_T GetTypeHash() const override
+	{
+		static size_t UniquePointer;
+		return reinterpret_cast<size_t>(&UniquePointer);
+	}
+
+private:
+	virtual void CreateRenderThreadResources() override
+	{
+		InstanceBuffer.bOwnedByEditor = Component->IsOwnedByEditor();
+	}
+	virtual void OnTransformChanged() override
+	{
+		TraversalOctree.Reset(new FLidarPointCloudTraversalOctree(Octree, Component->GetComponentTransform()));
+	}
+
+private:
+	FLidarPointCloudInstanceBuffer InstanceBuffer;
+	mutable TArray<FLidarPointCloudBatchElementUserData> UserData;
+	TUniquePtr<FLidarPointCloudTraversalOctree> TraversalOctree;
+	TArray<FLidarPointCloudTraversalOctreeNode*> SelectedNodes;
+
+	bool bReadyToRender;
+
+	ULidarPointCloudComponent* Component;
+	FLidarPointCloudOctree* Octree;
+	FMaterialRelevance MaterialRelevance;
+	AActor* Owner;
+
+	FLidarPointCloudCollisionRendering* CollisionRendering;
+};
+
+FPrimitiveSceneProxy* ULidarPointCloudComponent::CreateSceneProxy()
+{
+	FPrimitiveSceneProxy* Proxy = nullptr;
+	if (PointCloud)
+	{
+		Proxy = new FLidarPointCloudSceneProxy(this);
+	}
+	return Proxy;
+}
+
+void ULidarPointCloud::InitializeCollisionRendering()
+{
+	// Do not process, if the app in incapable of rendering
+	if (!FApp::CanEverRender())
+	{
+		return;
+	}
+
+	if (IsInRenderingThread())
+	{
+		FScopeLock Lock(&Octree.DataLock);
+
+		if (!CollisionRendering)
+		{
+			CollisionRendering = new FLidarPointCloudCollisionRendering();
+		}
+
+		CollisionRendering->Initialize(&Octree);
+	}
+	else
+	{
+		ENQUEUE_RENDER_COMMAND(InitializeCollisionRendering)(
+			[this](FRHICommandListImmediate& RHICmdList)
+			{
+				InitializeCollisionRendering();
+			});
+	}
+}
+
+void ULidarPointCloud::ReleaseCollisionRendering()
+{
+	// Do not process, if the app in incapable of rendering
+	if (!FApp::CanEverRender())
+	{
+		return;
+	}
+
+	if (IsInRenderingThread())
+	{
+		if (CollisionRendering)
+		{
+			delete CollisionRendering;
+			CollisionRendering = nullptr;
+		}
+	}
+	else
+	{
+		ENQUEUE_RENDER_COMMAND(ReleaseCollisionRendering)(
+			[this](FRHICommandListImmediate& RHICmdList)
+			{
+				ReleaseCollisionRendering();
+			});
+	}
+}
\ No newline at end of file
diff --git a/Engine/Plugins/Enterprise/LidarPointCloud/Source/LidarPointCloudRuntime/Public/IO/LidarPointCloudFileIO_LAS.h b/Engine/Plugins/Enterprise/LidarPointCloud/Source/LidarPointCloudRuntime/Public/IO/LidarPointCloudFileIO_LAS.h
new file mode 100644
--- /dev/null
+++ b/Engine/Plugins/Enterprise/LidarPointCloud/Source/LidarPointCloudRuntime/Public/IO/LidarPointCloudFileIO_LAS.h
@@ -0,0 +1,399 @@
+// Copyright Epic Games, Inc. All Rights Reserved.
+
+#pragma once
+
+#include "CoreMinimal.h"
+#include "LidarPointCloudShared.h"
+#include "LidarPointCloudFileIO.h"
+#include "Kismet/BlueprintFunctionLibrary.h"
+
+#if WITH_EDITOR
+#include "Widgets/SWidget.h"
+#endif
+
+// char[] { 'L', 'A', 'S', 'F' } read as uint32 - 1179861324
+#define HEADER_SIGNATURE 1179861324
+
+#include "LidarPointCloudFileIO_LAS.generated.h"
+
+#pragma pack(push)
+#pragma pack(1)
+
+/** Contains the common structure data between all formats */
+struct FLidarPointCloudFileIO_LAS_PointDataRecordFormatCommon
+{
+	FIntVector Location;
+	uint16 Intensity;
+};
+
+/** Contains the RGB extension structure data compatible with all formats */
+struct FLidarPointCloudFileIO_LAS_PointDataRecordFormatCommonRGB
+{
+	uint16 Red;
+	uint16 Green;
+	uint16 Blue;
+};
+
+/** Contains the Waveform extension structure data compatible with all formats */
+struct FLidarPointCloudFileIO_LAS_PointDataRecordFormatCommonWaveform
+{
+	uint8 WavePacketDescriptorIndex;
+	uint64 ByteOffsetToWaveformData;
+	uint32 WaveformPacketSize;
+	float ReturnPointWaveformLocation;
+	float ParametricDX;
+	float ParametricDY;
+	float ParametricDZ;
+};
+
+/**
+ * LAS - Point Data Record Format 0
+ * As per LAS Specification v.1.4 - R14
+ */
+struct FLidarPointCloudFileIO_LAS_PointDataRecordFormat0 : FLidarPointCloudFileIO_LAS_PointDataRecordFormatCommon
+{
+	uint8 ReturnNumber : 3;
+	uint8 NumberOfReturns : 3;
+	uint8 ScanDirectionFlag : 1;
+	uint8 EdgeOfFlightLine : 1;
+	uint8 Classification;
+	int8 ScanAngle;
+	uint8 UserData;
+	uint16 PointSourceID;
+};
+
+/**
+ * LAS - Point Data Record Format 1
+ * As per LAS Specification v.1.4 - R14
+ */
+struct FLidarPointCloudFileIO_LAS_PointDataRecordFormat1 : FLidarPointCloudFileIO_LAS_PointDataRecordFormat0
+{
+	double GPSTime;
+};
+
+/**
+ * LAS - Point Data Record Format 2
+ * As per LAS Specification v.1.4 - R14
+ */
+struct FLidarPointCloudFileIO_LAS_PointDataRecordFormat2 : FLidarPointCloudFileIO_LAS_PointDataRecordFormat0, FLidarPointCloudFileIO_LAS_PointDataRecordFormatCommonRGB { };
+
+/**
+ * LAS - Point Data Record Format 3
+ * As per LAS Specification v.1.4 - R14
+ */
+struct FLidarPointCloudFileIO_LAS_PointDataRecordFormat3 : FLidarPointCloudFileIO_LAS_PointDataRecordFormat1, FLidarPointCloudFileIO_LAS_PointDataRecordFormatCommonRGB { };
+
+/**
+ * LAS - Point Data Record Format 4
+ * As per LAS Specification v.1.4 - R14
+ */
+struct FLidarPointCloudFileIO_LAS_PointDataRecordFormat4 : FLidarPointCloudFileIO_LAS_PointDataRecordFormat1, FLidarPointCloudFileIO_LAS_PointDataRecordFormatCommonWaveform { };
+
+/**
+ * LAS - Point Data Record Format 5
+ * As per LAS Specification v.1.4 - R14
+ */
+struct FLidarPointCloudFileIO_LAS_PointDataRecordFormat5 : FLidarPointCloudFileIO_LAS_PointDataRecordFormat3, FLidarPointCloudFileIO_LAS_PointDataRecordFormatCommonWaveform { };
+
+/**
+ * LAS - Point Data Record Format 6
+ * As per LAS Specification v.1.4 - R14
+ */
+struct FLidarPointCloudFileIO_LAS_PointDataRecordFormat6 : FLidarPointCloudFileIO_LAS_PointDataRecordFormatCommon
+{
+	uint8 ReturnNumber : 4;
+	uint8 NumberOfReturns : 4;
+	uint8 ClassificationFlags : 4;
+	uint8 ScannerChannel : 2;
+	uint8 ScanDirectionFlag : 1;
+	uint8 EdgeOfFlightLine : 1;
+	uint8 Classification;
+	uint8 UserData;
+	int16 ScanAngle;
+	uint16 PointSourceID;
+	double GPSTime;
+};
+
+/**
+ * LAS - Point Data Record Format 7
+ * As per LAS Specification v.1.4 - R14
+ */
+struct FLidarPointCloudFileIO_LAS_PointDataRecordFormat7 : FLidarPointCloudFileIO_LAS_PointDataRecordFormat6, FLidarPointCloudFileIO_LAS_PointDataRecordFormatCommonRGB { };
+
+/**
+ * LAS - Point Data Record Format 8
+ * As per LAS Specification v.1.4 - R14
+ */
+struct FLidarPointCloudFileIO_LAS_PointDataRecordFormat8 : FLidarPointCloudFileIO_LAS_PointDataRecordFormat7
+{
+	uint16 NIR;
+};
+
+/**
+ * LAS - Point Data Record Format 9
+ * As per LAS Specification v.1.4 - R14
+ */
+struct FLidarPointCloudFileIO_LAS_PointDataRecordFormat9 : FLidarPointCloudFileIO_LAS_PointDataRecordFormat6, FLidarPointCloudFileIO_LAS_PointDataRecordFormatCommonWaveform { };
+
+/**
+ * LAS - Point Data Record Format 10
+ * As per LAS Specification v.1.4 - R14
+ */
+struct FLidarPointCloudFileIO_LAS_PointDataRecordFormat10 : FLidarPointCloudFileIO_LAS_PointDataRecordFormat8, FLidarPointCloudFileIO_LAS_PointDataRecordFormatCommonWaveform { };
+
+/**
+ * LAS - Public Header Block
+ * As per LAS Specification v.1.4 - R14
+ */
+struct FLidarPointCloudFileIO_LAS_PublicHeaderBlock
+{
+	/** Base set, 227 bytes */
+	uint32 FileSignature;	// Technically, this should be char[4] but uint32 is faster to compare
+	uint16 FileSourceID;
+	uint16 GlobalEncoding;
+	uint32 ProjectID_GUIDData1;
+	uint16 ProjectID_GUIDData2;
+	uint16 ProjectID_GUIDData3;
+	uint8 ProjectID_GUIDData4[8];
+	uint8 VersionMajor;
+	uint8 VersionMinor;
+	char SystemIdentifier[32];
+	char GeneratingSoftware[32];
+	uint16 FileCreationDayofYear;
+	uint16 FileCreationYear;
+	uint16 HeaderSize;
+	uint32 OffsetToPointData;
+	uint32 NumberOfVLRs;
+	uint8 PointDataRecordFormat;
+	uint16 PointDataRecordLength;
+	uint32 LegacyNumberOfPointRecords;
+	uint32 LegacyNumberOfPointsByReturn[5];
+	FDoubleVector ScaleFactor;
+	FDoubleVector Offset;
+	
+	/*
+	 * Order of data changes depending on the file version used
+	 * Pre 1.4: MaxX, MinX, MaxY, MinY, MaxZ, MinZ
+	 * 1.4+: MaxX, MaxY, MaxZ, MinX, MinY, MinZ
+	 */
+	double Bounds[6];
+
+	/** Added in 1.3, extra 8 bytes */
+	uint64 StartOfWaveformDataPacketRecord;
+
+	/** Added in 1.4, extra 140 bytes */
+	uint64 StartOfFirstEVLR;
+	uint32 NumberOfEVLRs;
+	uint64 NumberOfPointRecords;
+	uint64 NumberOfPointsByReturn[15];
+
+	FORCEINLINE bool IsValid() { return FileSignature == HEADER_SIGNATURE; }
+	FORCEINLINE uint64 GetNumberOfPoints() { return VersionMinor < 4 ? LegacyNumberOfPointRecords : NumberOfPointRecords; }
+
+	FORCEINLINE bool IsLegacyFormat() { return PointDataRecordFormat < 6; }
+
+	FORCEINLINE bool HasRGB()
+	{
+		switch (PointDataRecordFormat)
+		{
+		case 2:
+		case 3:
+		case 5:
+		case 7:
+		case 8:
+		case 10:
+			return true;
+
+		default:
+			return false;
+		}
+	}
+
+	FDoubleVector GetMin() const { return VersionMinor == 4 ? FDoubleVector(Bounds[3], Bounds[4], Bounds[5]) : FDoubleVector(Bounds[1], Bounds[3], Bounds[5]); }
+	FDoubleVector GetMax() const { return VersionMinor == 4 ? FDoubleVector(Bounds[0], Bounds[1], Bounds[2]) : FDoubleVector(Bounds[0], Bounds[2], Bounds[4]); }
+	FDoubleVector GetOrigin() const { return (GetMax() + GetMin()) / 2; }
+	FVector GetExtent() const { return ((GetMax() - GetMin()) / 2).ToVector(); }
+
+	void SetMin(const FDoubleVector& Min)
+	{
+		if (VersionMinor == 4)
+		{
+			Bounds[3] = Min.X;
+			Bounds[4] = Min.Y;
+			Bounds[5] = Min.Z;
+		}
+		else
+		{
+			Bounds[1] = Min.X;
+			Bounds[3] = Min.Y;
+			Bounds[5] = Min.Z;
+		}
+	}
+	void SetMax(const FDoubleVector& Max)
+	{
+		if (VersionMinor == 4)
+		{
+			Bounds[0] = Max.X;
+			Bounds[1] = Max.Y;
+			Bounds[2] = Max.Z;
+		}
+		else
+		{
+			Bounds[0] = Max.X;
+			Bounds[2] = Max.Y;
+			Bounds[4] = Max.Z;
+		}
+	}
+
+	/** Returns true if the point data contains user-specific extra bytes */
+	bool PointDataContainsExtraContent()
+	{
+		uint16 Size = 0;
+
+		switch (PointDataRecordFormat)
+		{
+		case 0: Size = sizeof(FLidarPointCloudFileIO_LAS_PointDataRecordFormat0); break;
+		case 1: Size = sizeof(FLidarPointCloudFileIO_LAS_PointDataRecordFormat1); break;
+		case 2: Size = sizeof(FLidarPointCloudFileIO_LAS_PointDataRecordFormat2); break;
+		case 3: Size = sizeof(FLidarPointCloudFileIO_LAS_PointDataRecordFormat3); break;
+		case 4: Size = sizeof(FLidarPointCloudFileIO_LAS_PointDataRecordFormat4); break;
+		case 5: Size = sizeof(FLidarPointCloudFileIO_LAS_PointDataRecordFormat5); break;
+		case 6: Size = sizeof(FLidarPointCloudFileIO_LAS_PointDataRecordFormat6); break;
+		case 7: Size = sizeof(FLidarPointCloudFileIO_LAS_PointDataRecordFormat7); break;
+		case 8: Size = sizeof(FLidarPointCloudFileIO_LAS_PointDataRecordFormat8); break;
+		case 9: Size = sizeof(FLidarPointCloudFileIO_LAS_PointDataRecordFormat9); break;
+		case 10: Size = sizeof(FLidarPointCloudFileIO_LAS_PointDataRecordFormat10); break;
+		}
+
+		return PointDataRecordLength != Size;
+	}
+
+	FString GetFormatDescription()
+	{
+		switch (PointDataRecordFormat)
+		{
+		case 0: return "Legacy Intensity";
+		case 1: return "Legacy Intensity with Time";
+		case 2: return "Legacy RGB";
+		case 3: return "Legacy RGB with Time";
+		case 4: return "Legacy Intensity with Time and Waveform";
+		case 5: return "Legacy RGB with Time and Waveform";
+		case 6: return "Intensity";
+		case 7: return "RGB";
+		case 8: return "RGB with NIR";
+		case 9: return "Intensity with Waveform";
+		case 10: return "RGB with NIR and Waveform";
+		default: return "Unknown";
+		}
+	}
+
+	static uint16 GetRecordLengthByFormat(uint8 Format)
+	{
+		switch (Format)
+		{
+		case 0: return sizeof(FLidarPointCloudFileIO_LAS_PointDataRecordFormat0);
+		case 1: return sizeof(FLidarPointCloudFileIO_LAS_PointDataRecordFormat1);
+		case 2: return sizeof(FLidarPointCloudFileIO_LAS_PointDataRecordFormat2);
+		case 3: return sizeof(FLidarPointCloudFileIO_LAS_PointDataRecordFormat3);
+		case 4: return sizeof(FLidarPointCloudFileIO_LAS_PointDataRecordFormat4);
+		case 5: return sizeof(FLidarPointCloudFileIO_LAS_PointDataRecordFormat5);
+		case 6: return sizeof(FLidarPointCloudFileIO_LAS_PointDataRecordFormat6);
+		case 7: return sizeof(FLidarPointCloudFileIO_LAS_PointDataRecordFormat7);
+		case 8: return sizeof(FLidarPointCloudFileIO_LAS_PointDataRecordFormat8);
+		case 9: return sizeof(FLidarPointCloudFileIO_LAS_PointDataRecordFormat9);
+		case 10: return sizeof(FLidarPointCloudFileIO_LAS_PointDataRecordFormat10);
+		default:
+			return 0;
+		}
+	}
+
+	/** Generates a new 1.2 compliant header, with pre-populated data. */
+	static FLidarPointCloudFileIO_LAS_PublicHeaderBlock Generate(const int64& NumberOfPoints, const FDoubleVector& Min, const FDoubleVector& Max)
+	{
+		FLidarPointCloudFileIO_LAS_PublicHeaderBlock Header;
+		FMemory::Memzero(&Header, sizeof(Header));
+
+		Header.FileSignature = HEADER_SIGNATURE;
+		Header.VersionMajor = 1;
+		Header.VersionMinor = 2;
+
+		const char* SystemIdentifier = "Unreal Engine 4";
+		FMemory::Memcpy(Header.SystemIdentifier, SystemIdentifier, 15);
+
+		const char* GeneratingSoftware = "Point Cloud Plugin";
+		FMemory::Memcpy(Header.GeneratingSoftware, GeneratingSoftware, 18);
+
+		FDateTime Date = FDateTime::Now();
+		Header.FileCreationDayofYear = Date.GetDayOfYear();
+		Header.FileCreationYear = Date.GetYear();
+
+		Header.HeaderSize = 227;
+		Header.OffsetToPointData = 227;
+
+		Header.PointDataRecordFormat = 2;
+		Header.PointDataRecordLength = GetRecordLengthByFormat(Header.PointDataRecordFormat);
+		Header.LegacyNumberOfPointRecords = NumberOfPoints;
+
+		FDoubleVector Size = Max - Min;
+		Header.Offset = Min;
+		Header.ScaleFactor = FDoubleVector(FMath::Pow(2, FMath::CeilToInt(FMath::Log2(Size.X)) - 31), FMath::Pow(2, FMath::CeilToInt(FMath::Log2(Size.Y)) - 31), FMath::Pow(2, FMath::CeilToInt(FMath::Log2(Size.Z)) - 31));
+		Header.SetMin(Min);
+		Header.SetMax(Max);
+
+		return Header;
+	}
+};
+#pragma pack(pop)
+
+struct FLidarPointCloudImportSettings_LAS : public FLidarPointCloudImportSettings
+{
+	FLidarPointCloudFileIO_LAS_PublicHeaderBlock PublicHeaderBlock;
+
+public:
+	FLidarPointCloudImportSettings_LAS(const FString& Filename);
+	virtual bool IsFileCompatible(const FString& InFilename) const override { return true; }
+	virtual void Serialize(FArchive& Ar) override;
+
+	virtual FString GetUID() const override { return "FLidarPointCloudImportSettings_LAS"; }
+
+	virtual void SetNewFilename(const FString& NewFilename) override
+	{
+		FLidarPointCloudImportSettings::SetNewFilename(NewFilename);
+		ReadFileHeader(NewFilename);
+	}
+
+	virtual TSharedPtr<FLidarPointCloudImportSettings> Clone(const FString& NewFilename = "") override
+	{
+		TSharedPtr<FLidarPointCloudImportSettings_LAS> NewSettings(new FLidarPointCloudImportSettings_LAS(NewFilename.IsEmpty() ? Filename : NewFilename));
+		NewSettings->bImportAll = bImportAll;
+		return NewSettings;
+	}
+
+private:
+	/** Reads and parses header information about the given file. */
+	void ReadFileHeader(const FString& InFilename);
+
+	friend class ULidarPointCloudFileIO_LAS;
+};
+
+/**
+ * Inherits from UBlueprintFunctionLibrary to allow exposure to Blueprint Library in the same class.
+ */
+UCLASS()
+class LIDARPOINTCLOUDRUNTIME_API ULidarPointCloudFileIO_LAS : public UBlueprintFunctionLibrary, public FLidarPointCloudFileIOHandler
+{
+	GENERATED_BODY()
+
+public:
+	virtual bool SupportsImport() const override { return true; }
+	virtual bool SupportsExport() const override { return true; }
+	
+	virtual TSharedPtr<FLidarPointCloudImportSettings> GetImportSettings(const FString& Filename) override { return TSharedPtr<FLidarPointCloudImportSettings>(new FLidarPointCloudImportSettings_LAS(Filename)); }
+
+	virtual bool HandleImport(const FString& Filename, TSharedPtr<FLidarPointCloudImportSettings> ImportSettings, FLidarPointCloudImportResults& OutImportResults) override;
+	virtual bool HandleExport(const FString& Filename, class ULidarPointCloud* PointCloud) override;
+
+	ULidarPointCloudFileIO_LAS() { ULidarPointCloudFileIO::RegisterHandler(this, { "LAS" }); }
+};
+
+#undef HEADER_SIGNATURE
\ No newline at end of file
diff --git a/Engine/Plugins/Enterprise/LidarPointCloud/Source/LidarPointCloudRuntime/Public/LidarPointCloud.h b/Engine/Plugins/Enterprise/LidarPointCloud/Source/LidarPointCloudRuntime/Public/LidarPointCloud.h
new file mode 100644
--- /dev/null
+++ b/Engine/Plugins/Enterprise/LidarPointCloud/Source/LidarPointCloudRuntime/Public/LidarPointCloud.h
@@ -0,0 +1,639 @@
+// Copyright Epic Games, Inc. All Rights Reserved.
+
+#pragma once
+
+#include "CoreMinimal.h"
+#include "Curves/CurveFloat.h"
+#include "LidarPointCloudShared.h"
+#include "LidarPointCloudOctree.h"
+#include "Widgets/Notifications/SNotificationList.h"
+#include "HAL/ThreadSafeBool.h"
+#include "Engine/EngineTypes.h"
+#include "LidarPointCloudSettings.h"
+#include "Kismet/BlueprintFunctionLibrary.h"
+#include "Engine/LatentActionManager.h"
+#include "ConvexVolume.h"
+#include "Interfaces/Interface_CollisionDataProvider.h"
+#include "LidarPointCloud.generated.h"
+
+class ALidarPointCloudActor;
+class ULidarPointCloudComponent;
+class UBodySetup;
+class FLidarPointCloudCollisionRendering;
+
+/**
+ * Used for ULidarPointCloud::CreateFromXXXX calls
+ */
+struct FLidarPointCloudAsyncParameters
+{
+	/** Should the creation use async operation */
+	bool bUseAsync;
+
+	/**
+	 * Called every time at least 1% progress is generated.
+	 * The parameter is within 0.0 - 1.0 range.
+	 */
+	TFunction<void(float)> ProgressCallback;
+
+	/**
+	 * Called once, when the operation completes.
+	 * The parameter specifies whether it has been executed successfully.
+	 */
+	TFunction<void(bool)> CompletionCallback;
+
+	FLidarPointCloudAsyncParameters(bool bUseAsync, TFunction<void(float)> ProgressCallback = nullptr, TFunction<void(bool)> CompletionCallback = nullptr)
+		: bUseAsync(bUseAsync)
+		, ProgressCallback(MoveTemp(ProgressCallback))
+		, CompletionCallback(MoveTemp(CompletionCallback))
+	{
+	}
+};
+
+/** Wrapper around a NotificationItem to make the notification handling more centralized */
+class FLidarPointCloudNotification
+{
+	/** Stores the pointer to the actual notification item */
+	TSharedPtr<SNotificationItem> NotificationItem;
+
+	/** Owner of this notification */
+	UObject* Owner;
+
+	FString CurrentText;
+	int8 CurrentProgress;
+
+public:
+	FLidarPointCloudNotification() : FLidarPointCloudNotification(nullptr) {}
+	FLidarPointCloudNotification(UObject* Owner);
+
+	bool IsValid() const { return NotificationItem.IsValid(); }
+	void Create(const FString& Text, FThreadSafeBool* bCancelPtr = nullptr, const FString& Icon = "ClassIcon32.LidarPointCloud");
+	void SetText(const FString& Text);
+	void SetProgress(int8 Progress);
+	void SetTextWithProgress(const FString& Text, int8 Progress);
+	void Close(bool bSuccess);
+
+private:
+	void UpdateStatus();
+};
+
+/**
+ * Represents the Point Cloud asset
+ */
+UCLASS(BlueprintType, AutoExpandCategories=("Performance", "Rendering|Sprite"), AutoCollapseCategories=("Import Settings"))
+class LIDARPOINTCLOUDRUNTIME_API ULidarPointCloud : public UObject, public IInterface_CollisionDataProvider
+{
+	GENERATED_BODY()
+
+private:
+	/** Stores the path to the original source file. Empty if dynamically created. */
+	UPROPERTY(EditAnywhere, Category = "Import Settings", meta = (AllowPrivateAccess = "true"))
+	FFilePath SourcePath;
+
+public:
+	/**
+	 * Determines the accuracy (in cm) of the collision for this point cloud.
+	 * NOTE: Lower values will require more time to build.
+	 * Rebuild collision for the changes to take effect.
+	 */
+	UPROPERTY(EditAnywhere, BlueprintReadWrite, Category = "Collision")
+	float CollisionAccuracy;
+
+	/** Holds pointer to the Import Settings used for the import. */
+	TSharedPtr<struct FLidarPointCloudImportSettings> ImportSettings;
+
+	FLidarPointCloudOctree Octree;
+	FLidarPointCloudCollisionRendering* CollisionRendering;
+
+	/** Stores the original offset as a double. */
+	UPROPERTY()
+	FDoubleVector OriginalCoordinates;
+
+	/** Required for file versioning */
+	static const FGuid PointCloudFileGUID;
+	static const int32 PointCloudFileVersion;
+
+private:
+	/** Used for caching the asset registry tag data. */
+	struct FLidarPointCloudAssetRegistryCache
+	{
+		FString PointCount;
+		FString ApproxSize;
+	} PointCloudAssetRegistryCache;
+
+	/** Contains the list of imported classification IDs */
+	UPROPERTY()
+	TArray<uint8> ClassificationsImported;
+
+	/** Used for async building */
+	FThreadSafeBool bAsyncCancelled;
+	FCriticalSection ImportLock;
+
+	/** Notification we hold on to, that indicates status and progress. */
+	FLidarPointCloudNotification Notification;
+
+	/** Description of collision */
+	UPROPERTY(transient, duplicatetransient)
+	UBodySetup* BodySetup;
+
+	/** Used for collision building */
+	FThreadSafeBool bCollisionBuildInProgress;
+
+public:
+	ULidarPointCloud();
+
+	/** Used to notify the component it should refresh its state. */
+	DECLARE_EVENT(ULidarPointCloud, FOnPointCloudChanged);
+	FOnPointCloudChanged& OnPointCloudRebuilt() { return OnPointCloudRebuiltEvent; }
+	FOnPointCloudChanged& OnPointCloudCollisionUpdated() { return OnPointCloudUpdateCollisionEvent; }
+
+	// Begin UObject Interface.
+	virtual void Serialize(FArchive& Ar) override;
+	virtual void PostLoad() override;
+	virtual void GetAssetRegistryTags(TArray<FAssetRegistryTag>& OutTags) const override;
+	virtual void BeginDestroy() override;
+#if WITH_EDITOR
+	virtual void PostEditChangeProperty(FPropertyChangedEvent & PropertyChangedEvent) override;
+#endif
+	// End UObject Interface.
+
+	UFUNCTION(BlueprintPure, Category = "Lidar Point Cloud")
+	int32 GetNumLODs() const { return Octree.GetNumLODs(); }
+
+	UFUNCTION(BlueprintPure, Category = "Lidar Point Cloud")
+	int32 GetNumPoints() const { return Octree.GetNumPoints(); }
+
+	UFUNCTION(BlueprintPure, Category = "Lidar Point Cloud")
+	int32 GetNumNodes() const { return Octree.GetNumNodes(); }
+
+	/** Returns the amount of memory in MB used to store the point cloud. */
+	UFUNCTION(BlueprintPure, Category = "Lidar Point Cloud")
+	int32 GetDataSize() const;
+
+	UFUNCTION(BlueprintPure, Category = "Lidar Point Cloud")
+	FString GetSourcePath() const { return SourcePath.FilePath; }
+
+	UFUNCTION(BlueprintPure, Category = "Lidar Point Cloud")
+	FBox GetBounds() const { return Octree.GetBounds(); }
+
+	UFUNCTION(BlueprintPure, Category = "Lidar Point Cloud")
+	FBox GetPointsBounds() const { return Octree.GetPointsBounds(); }
+
+	/** Recalculates and updates points bounds. */
+	UFUNCTION(BlueprintCallable, Category = "Lidar Point Cloud")
+	void RefreshPointsBounds() { Octree.RefreshPointsBounds(); }
+
+	/** Returns true, if the Octree has collision built */
+	UFUNCTION(BlueprintPure, Category = "Lidar Point Cloud")
+	bool HasCollisionData() const;
+
+	UFUNCTION(BlueprintCallable, Category = "Lidar Point Cloud")
+	void RefreshRendering() { OnPointCloudRebuiltEvent.Broadcast(); }
+
+	TArray<uint8> GetClassificationsImported() { return ClassificationsImported; }
+
+	/** Populates the given array with points from the tree */
+	TArray<FLidarPointCloudPoint*> GetPoints(int64 StartIndex = 0, int64 Count = -1);
+	void GetPoints(TArray<FLidarPointCloudPoint*>& Points, int64 StartIndex = 0, int64 Count = -1) { Octree.GetPoints(Points, StartIndex, Count); }
+
+	/** Populates the array with the list of points within the given sphere. */
+	void GetPointsInSphere(TArray<FLidarPointCloudPoint*>& SelectedPoints, const FVector& Center, const float& Radius, const bool& bVisibleOnly) { GetPointsInSphere(SelectedPoints, FSphere(Center, Radius), bVisibleOnly); }
+	void GetPointsInSphere(TArray<FLidarPointCloudPoint*>& SelectedPoints, const FSphere& Sphere, const bool& bVisibleOnly) { Octree.GetPointsInSphere(SelectedPoints, Sphere, bVisibleOnly); }
+
+	/** Populates the array with the list of points within the given box. */
+	void GetPointsInBox(TArray<FLidarPointCloudPoint*>& SelectedPoints, const FVector& Center, const FVector& Extent, const bool& bVisibleOnly) { GetPointsInBox(SelectedPoints, FBox(Center - Extent, Center + Extent), bVisibleOnly); }
+	void GetPointsInBox(TArray<FLidarPointCloudPoint*>& SelectedPoints, const FBox& Box, const bool& bVisibleOnly) { Octree.GetPointsInBox(SelectedPoints, Box, bVisibleOnly); }
+
+	/** Populates the array with the list of points within the given frustum. */
+	void GetPointsInFrustum(TArray<FLidarPointCloudPoint*>& SelectedPoints, const FConvexVolume& Frustum, const bool& bVisibleOnly) { Octree.GetPointsInFrustum(SelectedPoints, Frustum, bVisibleOnly); }
+
+	/** Returns an array with copies of points from the tree */
+	UFUNCTION(BlueprintPure, Category = "Lidar Point Cloud")
+	TArray<FLidarPointCloudPoint> GetPointsAsCopies(int32 StartIndex = 0, int32 Count = -1) const;
+	void GetPointsAsCopies(TArray<FLidarPointCloudPoint>& Points, int64 StartIndex = 0, int64 Count = -1) const { Octree.GetPointsAsCopies(Points, StartIndex, Count); }
+	
+	/** Returns an array with copies of points within the given sphere */
+	UFUNCTION(BlueprintPure, Category = "Lidar Point Cloud")
+	TArray<FLidarPointCloudPoint> GetPointsInSphereAsCopies(FVector Center, float Radius, bool bVisibleOnly);
+	void GetPointsInSphereAsCopies(TArray<FLidarPointCloudPoint>& SelectedPoints, const FSphere& Sphere, const bool& bVisibleOnly) { Octree.GetPointsInSphereAsCopies(SelectedPoints, Sphere, bVisibleOnly); }
+
+	/** Returns an array with copies of points within the given box */
+	UFUNCTION(BlueprintPure, Category = "Lidar Point Cloud")
+	TArray<FLidarPointCloudPoint> GetPointsInBoxAsCopies(FVector Center, FVector Extent, bool bVisibleOnly);
+	void GetPointsInBoxAsCopies(TArray<FLidarPointCloudPoint>& SelectedPoints, const FBox& Box, const bool& bVisibleOnly) { Octree.GetPointsInBoxAsCopies(SelectedPoints, Box, bVisibleOnly); }
+
+	/** Performs a raycast test against the point cloud. Returns the pointer if hit or nullptr otherwise. */
+	UFUNCTION(BlueprintPure, Category = "Lidar Point Cloud", meta = (Keywords = "raycast"))
+	bool LineTraceSingle(FVector Origin, FVector Direction, float Radius, bool bVisibleOnly, FLidarPointCloudPoint& PointHit);
+	FLidarPointCloudPoint* LineTraceSingle(const FLidarPointCloudRay& Ray, const float& Radius, const bool& bVisibleOnly) { return Octree.RaycastSingle(Ray, Radius, bVisibleOnly); }
+
+	/**
+	 * Performs a raycast test against the point cloud.
+	 * Populates OutHits array with the results.
+	 * Returns true it anything has been hit.
+	 */
+	UFUNCTION(BlueprintPure, Category = "Lidar Point Cloud", meta = (Keywords = "raycast"))
+	bool LineTraceMulti(FVector Origin, FVector Direction, float Radius, bool bVisibleOnly, TArray<FLidarPointCloudPoint>& OutHits) { return Octree.RaycastMulti(FLidarPointCloudRay(Origin, Direction), Radius, bVisibleOnly, OutHits); }
+	bool LineTraceMulti(const FLidarPointCloudRay& Ray, const float& Radius, bool bVisibleOnly, TArray<FLidarPointCloudPoint>& OutHits) { return Octree.RaycastMulti(Ray, Radius, bVisibleOnly, OutHits); }
+	bool LineTraceMulti(const FLidarPointCloudRay& Ray, const float& Radius, bool bVisibleOnly, TArray<FLidarPointCloudPoint*>& OutHits) { return Octree.RaycastMulti(Ray, Radius, bVisibleOnly, OutHits); }
+
+	/** Sets visibility of points within the given sphere. */
+	UFUNCTION(BlueprintCallable, Category = "Lidar Point Cloud")
+	void SetVisibilityOfPointsInSphere(bool bNewVisibility, FVector Center, float Radius) { SetVisibilityOfPointsInSphere(bNewVisibility, FSphere(Center, Radius)); }
+	void SetVisibilityOfPointsInSphere(const bool& bNewVisibility, const FSphere& Sphere) { Octree.SetVisibilityOfPointsInSphere(bNewVisibility, Sphere); }
+
+	/**
+	 * Sets visibility of points within the given sphere.
+	 * Async version - does not wait for completion before returning from the call.
+	 */
+	UFUNCTION(BlueprintCallable, Category = "Lidar Point Cloud")
+	void SetVisibilityOfPointsInSphereAsync(bool bNewVisibility, FVector Center, float Radius) { SetVisibilityOfPointsInSphereAsync(bNewVisibility, FSphere(Center, Radius)); }
+	void SetVisibilityOfPointsInSphereAsync(const bool& bNewVisibility, const FSphere& Sphere, TFunction<void(void)> CompletionCallback = nullptr) { Octree.SetVisibilityOfPointsInSphereAsync(bNewVisibility, Sphere, MoveTemp(CompletionCallback)); }
+
+	/** Sets visibility of points within the given box. */
+	UFUNCTION(BlueprintCallable, Category = "Lidar Point Cloud")
+	void SetVisibilityOfPointsInBox(bool bNewVisibility, FVector Center, FVector Extent) { SetVisibilityOfPointsInBox(bNewVisibility, FBox(Center - Extent, Center + Extent)); }
+	void SetVisibilityOfPointsInBox(const bool& bNewVisibility, const FBox& Box) { Octree.SetVisibilityOfPointsInBox(bNewVisibility, Box); }
+
+	/**
+	 * Sets visibility of points within the given box.
+	 * Async version - does not wait for completion before returning from the call.
+	 */
+	UFUNCTION(BlueprintCallable, Category = "Lidar Point Cloud")
+	void SetVisibilityOfPointsInBoxAsync(bool bNewVisibility, FVector Center, FVector Extent) { SetVisibilityOfPointsInBoxAsync(bNewVisibility, FBox(Center - Extent, Center + Extent)); }
+	void SetVisibilityOfPointsInBoxAsync(const bool& bNewVisibility, const FBox& Box, TFunction<void(void)> CompletionCallback = nullptr) { Octree.SetVisibilityOfPointsInBoxAsync(bNewVisibility, Box, MoveTemp(CompletionCallback)); }
+
+	/** Sets visibility of points hit by the given ray. */
+	UFUNCTION(BlueprintCallable, Category = "Lidar Point Cloud")
+	void SetVisibilityOfPointsByRay(bool bNewVisibility, FVector Origin, FVector Direction, float Radius) { SetVisibilityOfPointsByRay(bNewVisibility, FLidarPointCloudRay(Origin, Direction), Radius); }
+	void SetVisibilityOfPointsByRay(const bool& bNewVisibility, const FLidarPointCloudRay& Ray, const float& Radius) { Octree.SetVisibilityOfPointsByRay(bNewVisibility, Ray, Radius); }
+
+	/**
+	 * Sets visibility of points hit by the given ray.
+	 * Async version - does not wait for completion before returning from the call.
+	 */
+	UFUNCTION(BlueprintCallable, Category = "Lidar Point Cloud")
+	void SetVisibilityOfPointsByRayAsync(bool bNewVisibility, FVector Origin, FVector Direction, float Radius) { SetVisibilityOfPointsByRayAsync(bNewVisibility, FLidarPointCloudRay(Origin, Direction), Radius); }
+	void SetVisibilityOfPointsByRayAsync(const bool& bNewVisibility, const FLidarPointCloudRay& Ray, const float& Radius, TFunction<void(void)> CompletionCallback = nullptr) { Octree.SetVisibilityOfPointsByRayAsync(bNewVisibility, Ray, Radius, MoveTemp(CompletionCallback)); }
+
+	/** Marks all points hidden */
+	UFUNCTION(BlueprintCallable, Category = "Lidar Point Cloud")
+	void HideAll() { Octree.HideAll(); }
+
+	/** Marks all points visible */
+	UFUNCTION(BlueprintCallable, Category = "Lidar Point Cloud")
+	void UnhideAll() { Octree.UnhideAll(); }
+
+	/**
+	 * Marks all points visible
+	 * Async version - does not wait for completion before returning from the call.
+	 */
+	UFUNCTION(BlueprintCallable, Category = "Lidar Point Cloud")
+	void ResetVisibilityAsync() { Octree.ResetVisibilityAsync(); }
+	void ResetVisibilityAsync(TFunction<void(void)> CompletionCallback) { Octree.ResetVisibilityAsync(MoveTemp(CompletionCallback)); }
+
+	/** Executes the provided action on each of the points. */
+	void ExecuteActionOnAllPoints(TFunction<void(FLidarPointCloudPoint*)> Action, const bool& bVisibleOnly) { Octree.ExecuteActionOnAllPoints(MoveTemp(Action), bVisibleOnly); }
+
+	/**
+	 * Executes the provided action on each of the points.
+	 * Async version - does not wait for completion before returning from the call.
+	 */
+	void ExecuteActionOnAllPointsAsync(TFunction<void(FLidarPointCloudPoint*)> Action, const bool& bVisibleOnly, TFunction<void(void)> CompletionCallback = nullptr) { Octree.ExecuteActionOnAllPointsAsync(MoveTemp(Action), bVisibleOnly, MoveTemp(CompletionCallback)); }
+
+	/** Executes the provided action on each of the points within the given sphere. */
+	void ExecuteActionOnPointsInSphere(TFunction<void(FLidarPointCloudPoint*)> Action, const FVector& Center, const float& Radius, const bool& bVisibleOnly) { ExecuteActionOnPointsInSphere(MoveTemp(Action), FSphere(Center, Radius), bVisibleOnly); }
+	void ExecuteActionOnPointsInSphere(TFunction<void(FLidarPointCloudPoint*)> Action, const FSphere& Sphere, const bool& bVisibleOnly) { Octree.ExecuteActionOnPointsInSphere(MoveTemp(Action), Sphere, bVisibleOnly); }
+
+	/**
+	 * Executes the provided action on each of the points within the given sphere.
+	 * Async version - does not wait for completion before returning from the call.
+	 */
+	void ExecuteActionOnPointsInSphereAsync(TFunction<void(FLidarPointCloudPoint*)> Action, const FVector& Center, const float& Radius, const bool& bVisibleOnly) { ExecuteActionOnPointsInSphereAsync(MoveTemp(Action), FSphere(Center, Radius), bVisibleOnly); }
+	void ExecuteActionOnPointsInSphereAsync(TFunction<void(FLidarPointCloudPoint*)> Action, const FSphere& Sphere, const bool& bVisibleOnly, TFunction<void(void)> CompletionCallback = nullptr) { Octree.ExecuteActionOnPointsInSphereAsync(MoveTemp(Action), Sphere, bVisibleOnly, MoveTemp(CompletionCallback)); }
+
+	/** Executes the provided action on each of the points within the given box. */
+	void ExecuteActionOnPointsInBox(TFunction<void(FLidarPointCloudPoint*)> Action, const FVector& Center, const FVector& Extent, const bool& bVisibleOnly) { ExecuteActionOnPointsInBox(MoveTemp(Action), FBox(Center - Extent, Center + Extent), bVisibleOnly); }
+	void ExecuteActionOnPointsInBox(TFunction<void(FLidarPointCloudPoint*)> Action, const FBox& Box, const bool& bVisibleOnly) { Octree.ExecuteActionOnPointsInBox(MoveTemp(Action), Box, bVisibleOnly); }
+
+	/**
+	 * Executes the provided action on each of the points within the given box.
+	 * Async version - does not wait for completion before returning from the call.
+	 */
+	void ExecuteActionOnPointsInBoxAsync(TFunction<void(FLidarPointCloudPoint*)> Action, const FVector& Center, const FVector& Extent, const bool& bVisibleOnly) { ExecuteActionOnPointsInBoxAsync(MoveTemp(Action), FBox(Center - Extent, Center + Extent), bVisibleOnly); }
+	void ExecuteActionOnPointsInBoxAsync(TFunction<void(FLidarPointCloudPoint*)> Action, const FBox& Box, const bool& bVisibleOnly, TFunction<void(void)> CompletionCallback = nullptr) { Octree.ExecuteActionOnPointsInBoxAsync(MoveTemp(Action), Box, bVisibleOnly, MoveTemp(CompletionCallback)); }
+
+	/** Executes the provided action on each of the points hit by the given ray. */
+	void ExecuteActionOnPointsByRay(TFunction<void(FLidarPointCloudPoint*)> Action, const FLidarPointCloudRay& Ray, const float& Radius, bool bVisibleOnly) { Octree.ExecuteActionOnPointsByRay(MoveTemp(Action), Ray, Radius, bVisibleOnly); }
+
+	/**
+	 * Executes the provided action on each of the points hit by the given ray.
+	 * Async version - does not wait for completion before returning from the call.
+	 */
+	void ExecuteActionOnPointsByRayAsync(TFunction<void(FLidarPointCloudPoint*)> Action, const FLidarPointCloudRay& Ray, const float& Radius, bool bVisibleOnly, TFunction<void(void)> CompletionCallback = nullptr) { Octree.ExecuteActionOnPointsByRayAsync(MoveTemp(Action), Ray, Radius, bVisibleOnly, MoveTemp(CompletionCallback)); }
+
+	/** Applies the given color to all points */
+	UFUNCTION(BlueprintCallable, Category = "Lidar Point Cloud")
+	void ApplyColorToAllPoints(const FColor& NewColor, const bool& bVisibleOnly) { Octree.ApplyColorToAllPoints(NewColor, bVisibleOnly); }
+
+	/** Applies the given color to all points within the sphere */
+	UFUNCTION(BlueprintCallable, Category = "Lidar Point Cloud")
+	void ApplyColorToPointsInSphere(FColor NewColor, FVector Center, float Radius, bool bVisibleOnly) { ApplyColorToPointsInSphere(NewColor, FSphere(Center, Radius), bVisibleOnly); }
+	void ApplyColorToPointsInSphere(const FColor& NewColor, const FSphere& Sphere, const bool& bVisibleOnly) { Octree.ApplyColorToPointsInSphere(NewColor, Sphere, bVisibleOnly); }
+
+	/** Applies the given color to all points within the box */
+	UFUNCTION(BlueprintCallable, Category = "Lidar Point Cloud")
+	void ApplyColorToPointsInBox(FColor NewColor, FVector Center, FVector Extent, bool bVisibleOnly) { ApplyColorToPointsInBox(NewColor, FBox(Center - Extent, Center + Extent), bVisibleOnly); }
+	void ApplyColorToPointsInBox(const FColor& NewColor, const FBox& Box, const bool& bVisibleOnly) { Octree.ApplyColorToPointsInBox(NewColor, Box, bVisibleOnly); }
+
+	/** Applies the given color to all points hit by the given ray */
+	UFUNCTION(BlueprintCallable, Category = "Lidar Point Cloud")
+	void ApplyColorToPointsByRay(FColor NewColor, FVector Origin, FVector Direction, float Radius, bool bVisibleOnly) { ApplyColorToPointsByRay(NewColor, FLidarPointCloudRay(Origin, Direction), Radius, bVisibleOnly); }
+	void ApplyColorToPointsByRay(const FColor& NewColor, const FLidarPointCloudRay& Ray, const float& Radius, bool bVisibleOnly) { Octree.ApplyColorToPointsByRay(NewColor, Ray, Radius, bVisibleOnly); }
+
+	/**
+	 * This should to be called if any manual modification to individual points' visibility has been made.
+	 * If not marked dirty, the rendering may work sub-optimally.
+	 */
+	UFUNCTION(BlueprintCallable, Category = "Lidar Point Cloud")
+	void MarkPointVisibilityDirty() { Octree.MarkPointVisibilityDirty(); }
+
+	UFUNCTION(BlueprintCallable, Category = "Lidar Point Cloud")
+	void SetSourcePath(const FString& NewSourcePath);
+
+	/**
+	 * Re-initializes the asset with new bounds.
+	 * Warning: Will erase all currently held data!
+	 */
+	UFUNCTION(BlueprintCallable, Category = "Lidar Point Cloud")
+	void Initialize(const FBox& NewBounds) { Octree.Initialize(NewBounds); }
+
+	/** Builds collision mesh for the cloud, using current collision settings */
+	UFUNCTION(BlueprintCallable, Category = "Lidar Point Cloud")
+	void BuildCollision();
+
+	/** Removes collision mesh from the cloud. */
+	UFUNCTION(BlueprintCallable, Category = "Lidar Point Cloud")
+	void RemoveCollision();
+
+	/** Centers all contained points */
+	UFUNCTION(BlueprintCallable, Category = "Lidar Point Cloud")
+	void CenterPoints();
+
+	/** Restores original coordinates to all contained points */
+	UFUNCTION(BlueprintCallable, Category = "Lidar Point Cloud")
+	void RestoreOriginalCoordinates();
+
+	/**
+	 * Applies given offset to all contained points.
+	 * If bRefreshPointsBounds is set to false, make sure you call RefreshPointsBounds() manually or cloud centering may not work correctly.
+	 */
+	UFUNCTION(BlueprintCallable, Category = "Lidar Point Cloud")
+	void ShiftPointsBy(FVector Offset, bool bRefreshPointsBounds) { ShiftPointsBy(FDoubleVector(Offset), bRefreshPointsBounds); }
+	void ShiftPointsBy(FDoubleVector Offset, bool bRefreshPointsBounds);
+
+	/** Returns true, if the cloud has been centered. */
+	UFUNCTION(BlueprintCallable, Category = "Lidar Point Cloud")
+	bool IsCentered() const;
+
+	/** Re-imports the cloud from it's original source file, overwriting any current point information. */
+	UFUNCTION(BlueprintCallable, Category = "Lidar Point Cloud")
+	void Reimport(bool bUseAsync) { Reimport(FLidarPointCloudAsyncParameters(bUseAsync)); }
+	void Reimport(const FLidarPointCloudAsyncParameters& AsyncParameters);
+
+	/**
+	 * Exports this Point Cloud to the given filename.
+	 * Consult supported export formats.
+	 * Returns true if successful
+	 */
+	UFUNCTION(BlueprintCallable, Category = "Lidar Point Cloud")
+	bool Export(const FString& Filename);
+
+	/**
+	 * Inserts the given point into the Octree structure.
+	 * If bRefreshPointsBounds is set to false, make sure you call RefreshPointsBounds() manually or cloud centering may not work correctly.
+	 */
+	UFUNCTION(BlueprintCallable, Category = "Lidar Point Cloud")
+	void InsertPoint(const FLidarPointCloudPoint& Point, ELidarPointCloudDuplicateHandling DuplicateHandling, bool bRefreshPointsBounds);
+
+	/**
+	 * Inserts group of points into the Octree structure, multi-threaded.
+	 * If bRefreshPointsBounds is set to false, make sure you call RefreshPointsBounds() manually or cloud centering may not work correctly.
+	 */
+	UFUNCTION(BlueprintCallable, Category = "Lidar Point Cloud")
+	void InsertPoints(const TArray<FLidarPointCloudPoint>& Points, ELidarPointCloudDuplicateHandling DuplicateHandling, bool bRefreshPointsBounds) { InsertPoints(Points.GetData(), Points.Num(), DuplicateHandling, bRefreshPointsBounds); }
+
+	/**
+	 * Inserts group of points into the Octree structure, multi-threaded.
+	 * If bRefreshPointsBounds is set to false, make sure you call RefreshPointsBounds() manually or cloud centering may not work correctly.
+	 * Can be optionally passed a cancellation pointer - if it ever becomes non-null with value of true, process will be canceled.
+	 * May also provide progress callback, called approximately every 1% of progress.
+	 * Returns false if canceled.
+	 */
+	template<typename T>
+	bool InsertPoints(T InPoints, const int64& Count, ELidarPointCloudDuplicateHandling DuplicateHandling, bool bRefreshPointsBounds, FThreadSafeBool* bCanceled = nullptr, TFunction<void(float)> ProgressCallback = TFunction<void(float)>())
+	{
+		FScopeLock Lock(&Octree.DataLock);
+		return InsertPoints_NoLock(InPoints, Count, DuplicateHandling, bRefreshPointsBounds, bCanceled, MoveTemp(ProgressCallback));
+	}
+	template<typename T>
+	bool InsertPoints_NoLock(T InPoints, const int64& Count, ELidarPointCloudDuplicateHandling DuplicateHandling, bool bRefreshPointsBounds, FThreadSafeBool* bCanceled = nullptr, TFunction<void(float)> ProgressCallback = TFunction<void(float)>());
+
+	/**
+	 * Attempts to remove the given point.
+	 * If bRefreshPointsBounds is set to false, make sure you call RefreshPointsBounds() manually or cloud centering may not work correctly.
+	 */
+	UFUNCTION(BlueprintCallable, Category = "Lidar Point Cloud")
+	void RemovePoint(FLidarPointCloudPoint Point, bool bRefreshPointsBounds)
+	{
+		FScopeLock Lock(&Octree.DataLock);
+		Octree.RemovePoint(Point, bRefreshPointsBounds);
+	}
+	void RemovePoint_NoLock(FLidarPointCloudPoint Point, bool bRefreshPointsBounds)
+	{
+		Octree.RemovePoint(Point, bRefreshPointsBounds);
+	}
+	void RemovePoint(const FLidarPointCloudPoint* Point, bool bRefreshPointsBounds)
+	{
+		FScopeLock Lock(&Octree.DataLock);
+		RemovePoint_NoLock(Point, bRefreshPointsBounds);
+	}
+	void RemovePoint_NoLock(const FLidarPointCloudPoint* Point, bool bRefreshPointsBounds)
+	{
+		Octree.RemovePoint(Point, bRefreshPointsBounds);
+	}
+
+	/**
+	 * Removes points in bulk
+	 * If bRefreshPointsBounds is set to false, make sure you call RefreshPointsBounds() manually or cloud centering may not work correctly.
+	 */
+	void RemovePoints(TArray<FLidarPointCloudPoint*>& Points, bool bRefreshPointsBounds)
+	{
+		FScopeLock Lock(&Octree.DataLock);
+		RemovePoints_NoLock(Points, bRefreshPointsBounds);
+	}
+	void RemovePoints_NoLock(TArray<FLidarPointCloudPoint*>& Points, bool bRefreshPointsBounds)
+	{
+		Octree.RemovePoints(Points, bRefreshPointsBounds);
+	}
+
+	/**
+	 * Removes all points within the given sphere
+	 * If bRefreshPointsBounds is set to false, make sure you call RefreshPointsBounds() manually or cloud centering may not work correctly.
+	 */
+	UFUNCTION(BlueprintCallable, Category = "Lidar Point Cloud")
+	void RemovePointsInSphere(FVector Center, float Radius, bool bVisibleOnly, bool bRefreshPointsBounds) { RemovePointsInSphere(FSphere(Center, Radius), bVisibleOnly, bRefreshPointsBounds); }
+	void RemovePointsInSphere(const FSphere& Sphere, const bool& bVisibleOnly, bool bRefreshPointsBounds) { Octree.RemovePointsInSphere(Sphere, bVisibleOnly, bRefreshPointsBounds); }
+
+	/**
+	 * Removes all points within the given box
+	 * If bRefreshPointsBounds is set to false, make sure you call RefreshPointsBounds() manually or cloud centering may not work correctly.
+	 */
+	UFUNCTION(BlueprintCallable, Category = "Lidar Point Cloud")
+	void RemovePointsInBox(FVector Center, FVector Extent, bool bVisibleOnly, bool bRefreshPointsBounds) { RemovePointsInBox(FBox(Center - Extent, Center + Extent), bVisibleOnly, bRefreshPointsBounds); }
+	void RemovePointsInBox(const FBox& Box, const bool& bVisibleOnly, bool bRefreshPointsBounds) { Octree.RemovePointsInBox(Box, bVisibleOnly, bRefreshPointsBounds); }
+
+	/**
+	 * Removes all points hit by the given ray
+	 * If bRefreshPointsBounds is set to false, make sure you call RefreshPointsBounds() manually or cloud centering may not work correctly.
+	 */
+	UFUNCTION(BlueprintCallable, Category = "Lidar Point Cloud")
+	void RemovePointsByRay(FVector Origin, FVector Direction, float Radius, bool bVisibleOnly, bool bRefreshPointsBounds) { RemovePointsByRay(FLidarPointCloudRay(Origin, Direction), Radius, bVisibleOnly, bRefreshPointsBounds); }
+	void RemovePointsByRay(const FLidarPointCloudRay& Ray, const float& Radius, const bool& bVisibleOnly, bool bRefreshPointsBounds) { Octree.RemovePointsByRay(Ray, Radius, bVisibleOnly, bRefreshPointsBounds); }
+
+	/** Removes all hidden points */
+	UFUNCTION(BlueprintCallable, Category = "Lidar Point Cloud")
+	void RemoveHiddenPoints(bool bRefreshPointsBounds) { Octree.RemoveHiddenPoints(bRefreshPointsBounds); }
+
+	/** Reinitializes the cloud with the new set of data. */
+	UFUNCTION(BlueprintCallable, Category = "Lidar Point Cloud")
+	bool SetData(const TArray<FLidarPointCloudPoint>& Points) { return SetData(Points.GetData(), Points.Num()); }
+	bool SetData(TArray<FLidarPointCloudPoint*>& Points) { return SetData(Points.GetData(), Points.Num()); }
+	template<typename T>
+	bool SetData(T Points, const int64& Count, TFunction<void(float)> ProgressCallback = TFunction<void(float)>());
+
+	/** Merges this point cloud with the ones provided */
+	UFUNCTION(BlueprintCallable, Category = "Lidar Point Cloud")
+	void Merge(TArray<ULidarPointCloud*> PointCloudsToMerge) { Merge(PointCloudsToMerge, nullptr); }
+	void Merge(TArray<ULidarPointCloud*> PointCloudsToMerge, TFunction<void(float)> ProgressCallback);
+
+	/** Merges this point cloud with the one provided */
+	UFUNCTION(BlueprintCallable, Category = "Lidar Point Cloud")
+	void MergeSingle(ULidarPointCloud* PointCloudToMerge) { Merge(TArray<ULidarPointCloud*>({ PointCloudToMerge })); }
+
+	//~ Begin Interface_CollisionDataProvider Interface
+	virtual bool GetPhysicsTriMeshData(FTriMeshCollisionData* CollisionData, bool InUseAllTriData) override;
+	virtual bool ContainsPhysicsTriMeshData(bool InUseAllTriData) const override { return HasCollisionData(); }
+	virtual bool WantsNegXTriMesh() override { return false; }
+	//~ End Interface_CollisionDataProvider Interface
+
+	UBodySetup* GetBodySetup();
+
+public:
+	/** Aligns provided clouds based on the relative offset between their Original Coordinates. Retains overall centering of the group. */
+	static void AlignClouds(TArray<ULidarPointCloud*> PointCloudsToAlign);
+
+	/**
+	 * Returns new Point Cloud object imported using the settings provided.
+	 * Use nullptr as ImportSettings parameter to use default set of settings instead.
+	 */
+	static ULidarPointCloud* CreateFromFile(const FString& Filename, TSharedPtr<struct FLidarPointCloudImportSettings> ImportSettings = nullptr, UObject* InParent = (UObject*)GetTransientPackage(), FName InName = NAME_None, EObjectFlags Flags = RF_NoFlags)
+	{
+		return CreateFromFile(Filename, FLidarPointCloudAsyncParameters(GetDefault<ULidarPointCloudSettings>()->bUseAsyncImport), ImportSettings, InParent, InName, Flags);
+	}
+	static ULidarPointCloud* CreateFromFile(const FString& Filename, const FLidarPointCloudAsyncParameters& AsyncParameters, TSharedPtr<struct FLidarPointCloudImportSettings> ImportSettings = nullptr, UObject* InParent = (UObject*)GetTransientPackage(), FName InName = NAME_None, EObjectFlags Flags = RF_NoFlags);
+
+	/*
+	 * Returns new Point Cloud object created from the data provided.
+	 * Warning: If using Async, make sure the data does not get invalidated during processing!
+	 */
+	template<typename T>
+	static ULidarPointCloud* CreateFromData(T Points, const int64& Count, const FLidarPointCloudAsyncParameters& AsyncParameters);
+	static ULidarPointCloud* CreateFromData(const TArray<FLidarPointCloudPoint>& Points, const bool& bUseAsync) { return CreateFromData(Points.GetData(), Points.Num(), FLidarPointCloudAsyncParameters(bUseAsync)); }
+	static ULidarPointCloud* CreateFromData(const TArray<FLidarPointCloudPoint>& Points, const FLidarPointCloudAsyncParameters& AsyncParameters) { return CreateFromData(Points.GetData(), Points.Num(), AsyncParameters); }
+	static ULidarPointCloud* CreateFromData(TArray<FLidarPointCloudPoint*>& Points, const bool& bUseAsync) { return CreateFromData(Points.GetData(), Points.Num(), FLidarPointCloudAsyncParameters(bUseAsync)); }
+	static ULidarPointCloud* CreateFromData(TArray<FLidarPointCloudPoint*>& Points, const FLidarPointCloudAsyncParameters& AsyncParameters) { return CreateFromData(Points.GetData(), Points.Num(), AsyncParameters); }
+
+	/** Returns bounds fitting the given list of points */
+	static FBox CalculateBoundsFromPoints(const FLidarPointCloudPoint* Points, const int64& Count);
+	static FBox CalculateBoundsFromPoints(const TArray<FLidarPointCloudPoint>& Points) { return CalculateBoundsFromPoints(Points.GetData(), Points.Num()); }
+	static FBox CalculateBoundsFromPoints(FLidarPointCloudPoint** Points, const int64& Count);
+
+private:
+	/** Once async physics cook is done, create needed state */
+	void FinishPhysicsAsyncCook(UBodySetup* NewBodySetup) { FinishPhysicsAsyncCook(true, NewBodySetup); }
+	void FinishPhysicsAsyncCook(bool bSuccess, UBodySetup* NewBodySetup);
+
+	void InitializeCollisionRendering();
+	void ReleaseCollisionRendering();
+
+	FOnPointCloudChanged OnPointCloudRebuiltEvent;
+	FOnPointCloudChanged OnPointCloudUpdateCollisionEvent;
+};
+
+USTRUCT(BlueprintType)
+struct FLidarPointCloudTraceHit
+{
+	GENERATED_BODY()
+
+public:
+	UPROPERTY(EditAnywhere, BlueprintReadWrite, Category = "Collision")
+	ALidarPointCloudActor* Actor;
+
+	UPROPERTY(EditAnywhere, BlueprintReadWrite, Category = "Collision")
+	ULidarPointCloudComponent* Component;
+
+	UPROPERTY(EditAnywhere, BlueprintReadWrite, Category = "Collision")
+	TArray<FLidarPointCloudPoint> Points;
+
+	FLidarPointCloudTraceHit() : FLidarPointCloudTraceHit(nullptr, nullptr) {}
+	FLidarPointCloudTraceHit(ALidarPointCloudActor* Actor, ULidarPointCloudComponent* Component)
+		: Actor(Actor)
+		, Component(Component)
+	{
+	}
+};
+
+/**
+ * Blueprint library for the Point Cloud assets
+ */
+UCLASS(BlueprintType)
+class LIDARPOINTCLOUDRUNTIME_API ULidarPointCloudBlueprintLibrary : public UBlueprintFunctionLibrary
+{
+	GENERATED_BODY()
+
+public:
+	/** Returns new, empty Point Cloud object. */
+	UFUNCTION(BlueprintPure, Category = "Lidar Point Cloud", meta = (DisplayName = "Create Empty Lidar Point Cloud"))
+	static ULidarPointCloud* CreatePointCloudEmpty() { return NewObject<ULidarPointCloud>(); }
+
+	/**
+	 * Returns new Point Cloud object imported using default settings.
+	 * If using Async, the process runs in the background without blocking the game thread.
+	 */
+	UFUNCTION(BlueprintCallable, Category = "Lidar Point Cloud", meta = (Latent, WorldContext = "WorldContextObject", LatentInfo = "LatentInfo", ExpandEnumAsExecs = "AsyncMode", DisplayName = "Create Lidar Point Cloud From File"))
+	static void CreatePointCloudFromFile(UObject* WorldContextObject, const FString& Filename, bool bUseAsync, FLatentActionInfo LatentInfo, ELidarPointCloudAsyncMode& AsyncMode, float& Progress, ULidarPointCloud*& PointCloud);
+	static void CreatePointCloudFromFile(UObject* WorldContextObject, const FString& Filename, bool bUseAsync, FLatentActionInfo LatentInfo, TSharedPtr<struct FLidarPointCloudImportSettings> ImportSettings, ELidarPointCloudAsyncMode& AsyncMode, float& Progress, ULidarPointCloud*& PointCloud);
+
+	/*
+	 * Returns new Point Cloud object created from the data provided.
+	 * Warning: If using Async, make sure the data does not get invalidated during processing!
+	 */
+	UFUNCTION(BlueprintCallable, Category = "Lidar Point Cloud", meta = (Latent, WorldContext = "WorldContextObject", LatentInfo = "LatentInfo", ExpandEnumAsExecs = "AsyncMode", DisplayName = "Create Lidar Point Cloud From Data"))
+	static void CreatePointCloudFromData(UObject* WorldContextObject, const TArray<FLidarPointCloudPoint>& Points, bool bUseAsync, FLatentActionInfo LatentInfo, ELidarPointCloudAsyncMode& AsyncMode, float& Progress, ULidarPointCloud*& PointCloud);
+
+	/**
+	 * Exports the Point Cloud to the given filename.
+	 * Consult supported export formats.
+	 * Returns true if successful
+	 */
+	UFUNCTION(BlueprintCallable, Category = "Lidar Point Cloud")
+	static bool ExportPointCloudToFile(ULidarPointCloud* PointCloud, const FString& Filename) { return PointCloud && PointCloud->Export(Filename); }
+
+	/** Does a collision trace along the given line and returns the first blocking hit encountered. */
+	UFUNCTION(BlueprintPure, Category = "Lidar Point Cloud", meta = (WorldContext = "WorldContextObject", DisplayName = "LineTraceForLidarPointCloud", Keywords = "raycast"))
+	static bool LineTraceSingle(UObject* WorldContextObject, FVector Origin, FVector Direction, float Radius, bool bVisibleOnly, FLidarPointCloudTraceHit& Hit);
+
+	/** Does a collision trace along the given line and returns all hits encountered up to and including the first blocking hit. */
+	UFUNCTION(BlueprintPure, Category = "Lidar Point Cloud", meta = (WorldContext = "WorldContextObject", DisplayName = "LineTraceMultiForLidarPointCloud", Keywords = "raycast"))
+	static bool LineTraceMulti(UObject* WorldContextObject, FVector Origin, FVector Direction, float Radius, bool bVisibleOnly, TArray<FLidarPointCloudTraceHit>& Hits);
+};
\ No newline at end of file
diff --git a/Engine/Plugins/Enterprise/LidarPointCloud/Source/LidarPointCloudRuntime/Public/LidarPointCloudComponent.h b/Engine/Plugins/Enterprise/LidarPointCloud/Source/LidarPointCloudRuntime/Public/LidarPointCloudComponent.h
new file mode 100644
--- /dev/null
+++ b/Engine/Plugins/Enterprise/LidarPointCloud/Source/LidarPointCloudRuntime/Public/LidarPointCloudComponent.h
@@ -0,0 +1,535 @@
+// Copyright Epic Games, Inc. All Rights Reserved.
+
+#pragma once
+
+#include "CoreMinimal.h"
+#include "LidarPointCloud.h"
+#include "Components/MeshComponent.h"
+#include "LidarPointCloudComponent.generated.h"
+
+class UBodySetup;
+
+UENUM(BlueprintType)
+enum class ELidarPointCloudColorationMode : uint8
+{
+	/** Uses color tint only */
+	None,
+	/** Uses imported RGB / Intensity data */
+	Data,
+	/** The cloud's color will be overridden with elevation-based color */
+	Elevation,
+	/** The cloud's color will be overridden with relative position-based color */
+	Position,
+	/** Uses Classification ID of the point along with the component's Classification Colors property to sample the color */
+	Classification
+};
+
+UENUM(BlueprintType)
+enum class ELidarPointCloudSpriteShape : uint8
+{
+	Square,
+	Circle,
+};
+
+/** Component that allows you to render specified point cloud section */
+UCLASS(ClassGroup=Rendering, ShowCategories = (Rendering), HideCategories = (Object, LOD, Physics, Activation, Materials, Cooking, Input, HLOD, Mobile), meta = (BlueprintSpawnableComponent))
+class LIDARPOINTCLOUDRUNTIME_API ULidarPointCloudComponent : public UMeshComponent
+{
+	GENERATED_BODY()
+		
+private:
+	UPROPERTY(EditAnywhere, BlueprintReadOnly, Category = "Lidar Point Cloud", meta = (AllowPrivateAccess = "true"))
+	ULidarPointCloud *PointCloud;
+
+	/**
+	 * Allows using custom-built material for the point cloud.
+	 * Set to None to use the default one instead.
+	 */
+	UPROPERTY(EditAnywhere, BlueprintReadOnly, Category = "Material", meta = (AllowPrivateAccess = "true"))
+	UMaterialInterface* CustomMaterial;
+
+public:
+	/**
+	 * Determines the maximum number of points to be visible on the screen.
+	 * Higher values will produce better image quality, but will require faster hardware.
+	 */
+	UPROPERTY(EditAnywhere, BlueprintReadWrite, Category = "Performance", meta = (ClampMin = "0"))
+	int32 PointBudget;
+
+	/**
+	 * Determines the preference towards selecting nodes closer to screen center.
+	 * Larger values assign more priority towards screen center.
+	 * 0 to disable.
+	 */
+	UPROPERTY(EditAnywhere, BlueprintReadWrite, AdvancedDisplay, Category = "Performance", meta = (ClampMin = "0.0"))
+	float ScreenCenterImportance;
+
+	/**
+	 * Determines the minimum screen size for the node to be rendered.
+	 * Lower values will produce farther view distance, at the cost of higher CPU usage.
+	 */
+	UPROPERTY(EditAnywhere, BlueprintReadWrite, AdvancedDisplay, Category = "Performance", meta = (ClampMin = "0"))
+	float MinScreenSize;
+
+	/**
+	 * Use to tweak the size of the points.
+	 * Set to 0 to switch to 1 pixel points.
+	 */
+	UPROPERTY(EditAnywhere, BlueprintReadWrite, Category = "Appearance", meta = (ClampMin = "0.0"))
+	float PointSize;
+
+	/** Specifies which source to use for point colors. */
+	UPROPERTY(EditAnywhere, BlueprintReadWrite, Category = "Appearance")
+	ELidarPointCloudColorationMode ColorSource;
+
+	/** Affects the shape of points. */
+	UPROPERTY(EditAnywhere, BlueprintReadWrite, Category = "Appearance")
+	ELidarPointCloudSpriteShape PointShape;
+
+	/**
+	 * Used with the Classification source.
+	 * Maps the given classification ID to a color.
+	 */
+	UPROPERTY(EditAnywhere, BlueprintReadWrite, Category = "Appearance")
+	TMap<int32, FLinearColor> ClassificationColors;
+
+	/** Specifies the bottom color of the elevation-based gradient. */
+	UPROPERTY(EditAnywhere, BlueprintReadWrite, AdvancedDisplay, Category = "Appearance")
+	FLinearColor ElevationColorBottom;
+
+	/** Specifies the top color of the elevation-based gradient. */
+	UPROPERTY(EditAnywhere, BlueprintReadWrite, AdvancedDisplay, Category = "Appearance")
+	FLinearColor ElevationColorTop;
+
+	/**
+	 * Larger values will help mask LOD transition areas, but too large values will lead to loss of detail.
+	 * Values in range 0.035 - 0.05 seem to produce best overall results.
+	 */
+	UPROPERTY(EditAnywhere, BlueprintReadWrite, AdvancedDisplay, Category = "Appearance", meta = (ClampMin = "0.0", ClampMax = "0.15"))
+	float PointSizeBias;
+
+	UPROPERTY(interp, EditAnywhere, BlueprintReadWrite, Category = "Color Adjustment", meta = (UIMin = "0.0", UIMax = "2.0", Delta = "0.01", ColorGradingMode = "saturation", ShiftMouseMovePixelPerDelta = "10"))
+	FVector4 Saturation;
+
+	UPROPERTY(interp, EditAnywhere, BlueprintReadWrite, Category = "Color Adjustment", meta = (UIMin = "0.0", UIMax = "2.0", Delta = "0.01", ColorGradingMode = "contrast", ShiftMouseMovePixelPerDelta = "10"))
+	FVector4 Contrast;
+
+	UPROPERTY(interp, EditAnywhere, BlueprintReadWrite, Category = "Color Adjustment", meta = (UIMin = "0.0", UIMax = "2.0", Delta = "0.01", ColorGradingMode = "gamma", ShiftMouseMovePixelPerDelta = "10"))
+	FVector4 Gamma;
+
+	/** Affects the emissive strength of the color. Useful to create Bloom and light bleed effects. */
+	UPROPERTY(interp, EditAnywhere, BlueprintReadWrite, Category = "Color Adjustment", meta = (UIMin = "0.0", UIMax = "1.0", Delta = "0.01", ColorGradingMode = "gain", ShiftMouseMovePixelPerDelta = "10"))
+	FVector4 Gain;
+
+	/** Applied additively, 0 being neutral. */
+	UPROPERTY(interp, EditAnywhere, BlueprintReadWrite, Category = "Color Adjustment", meta = (UIMin = "-1.0", UIMax = "1.0", Delta = "0.001", ColorGradingMode = "offset", ShiftMouseMovePixelPerDelta = "20", SupportDynamicSliderMaxValue = "true", SupportDynamicSliderMinValue = "true"))
+	FVector4 Offset;
+
+	/** Specifies the tint to apply to the points. */
+	UPROPERTY(EditAnywhere, BlueprintReadWrite, Category = "Color Adjustment")
+	FLinearColor ColorTint;
+
+	/** Specifies the influence of Intensity data, if available, on the overall color. */
+	UPROPERTY(EditAnywhere, BlueprintReadWrite, Category = "Color Adjustment", meta = (ClampMin = "0", ClampMax = "1"))
+	float IntensityInfluence;
+
+	/**
+	 * Minimum Depth from which the nodes should be rendered.
+	 * 0 to disable.
+	 */
+	UPROPERTY(EditAnywhere, BlueprintReadWrite, AdvancedDisplay, Category = "Rendering", meta = (ClampMin = "0"))
+	int32 MinDepth;
+
+	/**
+	 * Maximum Depth to which the nodes should be rendered.
+	 * -1 to disable.
+	 */
+	UPROPERTY(EditAnywhere, BlueprintReadWrite, AdvancedDisplay, Category = "Rendering")
+	int32 MaxDepth;
+
+	/** Enabling this will cause the visible nodes to render their bounds. */
+	UPROPERTY(EditAnywhere, BlueprintReadWrite, AdvancedDisplay, Category = "Rendering")
+	bool bDrawNodeBounds;
+
+private:
+	UPROPERTY(Transient)
+	UMaterialInterface* Material;
+	UPROPERTY(Transient)
+	UMaterialInterface *MasterMaterial;
+	UPROPERTY(Transient)
+	UMaterialInterface* MasterMaterialMasked;
+
+	/** For asset editor use only */
+	bool bOwnedByEditor;
+
+public:
+	ULidarPointCloudComponent();
+
+	UFUNCTION(BlueprintPure, Category = "Components|LidarPointCloud")
+	ULidarPointCloud* GetPointCloud() const { return PointCloud; }
+
+	FORCEINLINE bool IsOwnedByEditor() const { return bOwnedByEditor; }
+
+	/** Populates the array with the list of points within the given sphere. */
+	void GetPointsInSphere(TArray<FLidarPointCloudPoint*>& SelectedPoints, const FVector& Center, const float& Radius, const bool& bVisibleOnly) { GetPointsInSphere(SelectedPoints, FSphere(Center, Radius), bVisibleOnly); }
+	void GetPointsInSphere(TArray<FLidarPointCloudPoint*>& SelectedPoints, const FSphere& Sphere, const bool& bVisibleOnly)
+	{
+		if (PointCloud)
+		{
+			PointCloud->GetPointsInSphere(SelectedPoints, Sphere.TransformBy(GetComponentTransform().Inverse()), bVisibleOnly);
+		}
+	}
+
+	/** Populates the array with the list of points within the given box. */
+	void GetPointsInBox(TArray<FLidarPointCloudPoint*>& SelectedPoints, const FVector& Center, const FVector& Extent, const bool& bVisibleOnly) { GetPointsInBox(SelectedPoints, FBox(Center - Extent, Center + Extent), bVisibleOnly); }
+	void GetPointsInBox(TArray<FLidarPointCloudPoint*>& SelectedPoints, const FBox& Box, const bool& bVisibleOnly)
+	{
+		if (PointCloud)
+		{
+			PointCloud->GetPointsInBox(SelectedPoints, Box.TransformBy(GetComponentTransform().Inverse()), bVisibleOnly);
+		}
+	}
+
+	/** Populates the array with copies of points within the given sphere. */
+	UFUNCTION(BlueprintCallable, Category = "Lidar Point Cloud")
+	TArray<FLidarPointCloudPoint> GetPointsInSphereAsCopies(const FVector& Center, const float& Radius, const bool& bVisibleOnly)
+	{
+		TArray<FLidarPointCloudPoint> Points;
+		GetPointsInSphereAsCopies(Points, FSphere(Center, Radius), bVisibleOnly);
+		return Points;
+	}
+	void GetPointsInSphereAsCopies(TArray<FLidarPointCloudPoint>& SelectedPoints, const FSphere& Sphere, const bool& bVisibleOnly)
+	{
+		if (PointCloud)
+		{
+			PointCloud->GetPointsInSphereAsCopies(SelectedPoints, Sphere.TransformBy(GetComponentTransform().Inverse()), bVisibleOnly);
+		}
+	}
+
+	/** Populates the array with copies of points within the given box. */
+	UFUNCTION(BlueprintCallable, Category = "Lidar Point Cloud")
+	TArray<FLidarPointCloudPoint> GetPointsInBoxAsCopies(FVector Center, FVector Extent, bool bVisibleOnly)
+	{
+		TArray<FLidarPointCloudPoint> Points;
+		GetPointsInBoxAsCopies(Points, FBox(Center - Extent, Center + Extent), bVisibleOnly);
+		return Points;
+	}
+	void GetPointsInBoxAsCopies(TArray<FLidarPointCloudPoint>& SelectedPoints, const FBox& Box, const bool& bVisibleOnly)
+	{
+		if (PointCloud)
+		{
+			PointCloud->GetPointsInBoxAsCopies(SelectedPoints, Box.TransformBy(GetComponentTransform().Inverse()), bVisibleOnly);
+		}
+	}
+
+	/** Performs a raycast test against the point cloud. Returns the pointer if hit or nullptr otherwise. */
+	UFUNCTION(BlueprintPure, Category = "Lidar Point Cloud", meta = (Keywords = "raycast"))
+	bool LineTraceSingle(FVector Origin, FVector Direction, float Radius, bool bVisibleOnly, FLidarPointCloudPoint& PointHit)
+	{
+		FLidarPointCloudPoint* Point = LineTraceSingle(FLidarPointCloudRay(Origin, Direction), Radius, bVisibleOnly);
+		if (Point)
+		{
+			PointHit = *Point;
+			return true;
+		}
+
+		return false;
+	}
+	FLidarPointCloudPoint* LineTraceSingle(FLidarPointCloudRay Ray, float Radius, bool bVisibleOnly)
+	{
+		return PointCloud ? PointCloud->LineTraceSingle(Ray.TransformBy(GetComponentTransform().Inverse()), Radius, bVisibleOnly) : nullptr;
+	}
+
+	/**
+	 * Performs a raycast test against the point cloud.
+	 * Populates OutHits array with the results.
+	 * Returns true it anything has been hit.
+	 */
+	UFUNCTION(BlueprintPure, Category = "Lidar Point Cloud", meta = (Keywords = "raycast"))
+	bool LineTraceMulti(FVector Origin, FVector Direction, float Radius, bool bVisibleOnly, TArray<FLidarPointCloudPoint>& OutHits)
+	{
+		return PointCloud ? PointCloud->LineTraceMulti(FLidarPointCloudRay(Origin, Direction).TransformBy(GetComponentTransform().Inverse()), Radius, bVisibleOnly, OutHits) : false;
+	}
+	bool LineTraceMulti(FLidarPointCloudRay Ray, float Radius, bool bVisibleOnly, TArray<FLidarPointCloudPoint>& OutHits)
+	{
+		return PointCloud ? PointCloud->LineTraceMulti(Ray.TransformBy(GetComponentTransform().Inverse()), Radius, bVisibleOnly, OutHits) : false;
+	}
+	bool LineTraceMulti(FLidarPointCloudRay Ray, float Radius, bool bVisibleOnly, TArray<FLidarPointCloudPoint*>& OutHits)
+	{
+		return PointCloud ? PointCloud->LineTraceMulti(Ray.TransformBy(GetComponentTransform().Inverse()), Radius, bVisibleOnly, OutHits) : false;
+	}
+
+	/** Sets visibility of points within the given sphere. */
+	UFUNCTION(BlueprintCallable, Category = "Lidar Point Cloud")
+	void SetVisibilityOfPointsInSphere(bool bNewVisibility, FVector Center, float Radius) { SetVisibilityOfPointsInSphere(bNewVisibility, FSphere(Center, Radius)); }
+	void SetVisibilityOfPointsInSphere(const bool& bNewVisibility, const FSphere& Sphere)
+	{
+		if (PointCloud)
+		{
+			PointCloud->SetVisibilityOfPointsInSphere(bNewVisibility, Sphere.TransformBy(GetComponentTransform().Inverse()));
+		}
+	}
+
+	/**
+	 * Sets visibility of points within the given sphere.
+	 * Async version - does not wait for completion before returning from the call.
+	 */
+	UFUNCTION(BlueprintCallable, Category = "Lidar Point Cloud")
+	void SetVisibilityOfPointsInSphereAsync(bool bNewVisibility, FVector Center, float Radius) { SetVisibilityOfPointsInSphereAsync(bNewVisibility, FSphere(Center, Radius)); }
+	void SetVisibilityOfPointsInSphereAsync(const bool& bNewVisibility, const FSphere& Sphere)
+	{
+		if (PointCloud)
+		{
+			PointCloud->SetVisibilityOfPointsInSphereAsync(bNewVisibility, Sphere.TransformBy(GetComponentTransform().Inverse()));
+		}
+	}
+
+	/** Sets visibility of points within the given box. */
+	UFUNCTION(BlueprintCallable, Category = "Lidar Point Cloud")
+	void SetVisibilityOfPointsInBox(bool bNewVisibility, FVector Center, FVector Extent) { SetVisibilityOfPointsInBox(bNewVisibility, FBox(Center - Extent, Center + Extent)); }
+	void SetVisibilityOfPointsInBox(const bool& bNewVisibility, const FBox& Box)
+	{
+		if (PointCloud)
+		{
+			PointCloud->SetVisibilityOfPointsInBox(bNewVisibility, Box.TransformBy(GetComponentTransform().Inverse()));
+		}
+	}
+
+	/**
+	 * Sets visibility of points within the given box.
+	 * Async version - does not wait for completion before returning from the call.
+	 */
+	UFUNCTION(BlueprintCallable, Category = "Lidar Point Cloud")
+	void SetVisibilityOfPointsInBoxAsync(bool bNewVisibility, FVector Center, FVector Extent) { SetVisibilityOfPointsInBoxAsync(bNewVisibility, FBox(Center - Extent, Center + Extent)); }
+	void SetVisibilityOfPointsInBoxAsync(const bool& bNewVisibility, const FBox& Box)
+	{
+		if (PointCloud)
+		{
+			PointCloud->SetVisibilityOfPointsInBoxAsync(bNewVisibility, Box.TransformBy(GetComponentTransform().Inverse()));
+		}
+	}
+
+	/** Sets visibility of points hit by the given ray. */
+	UFUNCTION(BlueprintCallable, Category = "Lidar Point Cloud")
+	void SetVisibilityOfPointsByRay(bool bNewVisibility, FVector Origin, FVector Direction, float Radius) { SetVisibilityOfPointsByRay(bNewVisibility, FLidarPointCloudRay(Origin, Direction), Radius); }
+	void SetVisibilityOfPointsByRay(bool bNewVisibility, FLidarPointCloudRay Ray, float Radius)
+	{
+		if (PointCloud)
+		{
+			PointCloud->SetVisibilityOfPointsByRay(bNewVisibility, Ray.TransformBy(GetComponentTransform().Inverse()), Radius);
+		}
+	}
+
+	/**
+	 * Sets visibility of points hit by the given ray.
+	 * Async version - does not wait for completion before returning from the call.
+	 */
+	UFUNCTION(BlueprintCallable, Category = "Lidar Point Cloud")
+	void SetVisibilityOfPointsByRayAsync(bool bNewVisibility, FVector Origin, FVector Direction, float Radius) { SetVisibilityOfPointsByRayAsync(bNewVisibility, FLidarPointCloudRay(Origin, Direction), Radius); }
+	void SetVisibilityOfPointsByRayAsync(bool bNewVisibility, FLidarPointCloudRay Ray, float Radius, TFunction<void(void)> CompletionCallback = nullptr)
+	{
+		if (PointCloud)
+		{
+			PointCloud->SetVisibilityOfPointsByRayAsync(bNewVisibility, Ray.TransformBy(GetComponentTransform().Inverse()), Radius, MoveTemp(CompletionCallback));
+		}
+	}
+
+	/** Executes the provided action on each of the points within the given sphere. */
+	void ExecuteActionOnPointsInSphere(TFunction<void(FLidarPointCloudPoint*)> Action, const FVector& Center, const float& Radius, const bool& bVisibleOnly) { ExecuteActionOnPointsInSphere(MoveTemp(Action), FSphere(Center, Radius), bVisibleOnly); }
+	void ExecuteActionOnPointsInSphere(TFunction<void(FLidarPointCloudPoint*)> Action, const FSphere& Sphere, const bool& bVisibleOnly)
+	{
+		if (PointCloud)
+		{
+			PointCloud->ExecuteActionOnPointsInSphere(MoveTemp(Action), Sphere.TransformBy(GetComponentTransform().Inverse()), bVisibleOnly);
+		}
+	}
+
+	/**
+	 * Executes the provided action on each of the points within the given sphere.
+	 * Async version - does not wait for completion before returning from the call.
+	 */
+	void ExecuteActionOnPointsInSphereAsync(TFunction<void(FLidarPointCloudPoint*)> Action, const FVector& Center, const float& Radius, const bool& bVisibleOnly) { ExecuteActionOnPointsInSphereAsync(MoveTemp(Action), FSphere(Center, Radius), bVisibleOnly); }
+	void ExecuteActionOnPointsInSphereAsync(TFunction<void(FLidarPointCloudPoint*)> Action, const FSphere& Sphere, const bool& bVisibleOnly)
+	{
+		if (PointCloud)
+		{
+			PointCloud->ExecuteActionOnPointsInSphereAsync(MoveTemp(Action), Sphere.TransformBy(GetComponentTransform().Inverse()), bVisibleOnly);
+		}
+	}
+
+	/** Executes the provided action on each of the points within the given box. */
+	void ExecuteActionOnPointsInBox(TFunction<void(FLidarPointCloudPoint*)> Action, const FVector& Center, const FVector& Extent, const bool& bVisibleOnly) { ExecuteActionOnPointsInBox(MoveTemp(Action), FBox(Center - Extent, Center + Extent), bVisibleOnly); }
+	void ExecuteActionOnPointsInBox(TFunction<void(FLidarPointCloudPoint*)> Action, const FBox& Box, const bool& bVisibleOnly)
+	{
+		if (PointCloud)
+		{
+			PointCloud->ExecuteActionOnPointsInBox(MoveTemp(Action), Box.TransformBy(GetComponentTransform().Inverse()), bVisibleOnly);
+		}
+	}
+
+	/**
+	 * Executes the provided action on each of the points within the given box.
+	 * Async version - does not wait for completion before returning from the call.
+	 */
+	void ExecuteActionOnPointsInBoxAsync(TFunction<void(FLidarPointCloudPoint*)> Action, const FVector& Center, const FVector& Extent, const bool& bVisibleOnly) { ExecuteActionOnPointsInBoxAsync(MoveTemp(Action), FBox(Center - Extent, Center + Extent), bVisibleOnly); }
+	void ExecuteActionOnPointsInBoxAsync(TFunction<void(FLidarPointCloudPoint*)> Action, const FBox& Box, const bool& bVisibleOnly)
+	{
+		if (PointCloud)
+		{
+			PointCloud->ExecuteActionOnPointsInBoxAsync(MoveTemp(Action), Box.TransformBy(GetComponentTransform().Inverse()), bVisibleOnly);
+		}
+	}
+
+	/** Executes the provided action on each of the points hit by the given ray. */
+	void ExecuteActionOnPointsByRay(TFunction<void(FLidarPointCloudPoint*)> Action, FLidarPointCloudRay Ray, float Radius, bool bVisibleOnly)
+	{
+		if (PointCloud)
+		{
+			PointCloud->ExecuteActionOnPointsByRay(MoveTemp(Action), Ray.TransformBy(GetComponentTransform().Inverse()), Radius, bVisibleOnly);
+		}
+	}
+
+	/**
+	 * Executes the provided action on each of the points hit by the given ray.
+	 * Async version - does not wait for completion before returning from the call.
+	 */
+	void ExecuteActionOnPointsByRayAsync(TFunction<void(FLidarPointCloudPoint*)> Action, FLidarPointCloudRay Ray, float Radius, bool bVisibleOnly, TFunction<void(void)> CompletionCallback = nullptr)
+	{
+		if (PointCloud)
+		{
+			PointCloud->ExecuteActionOnPointsByRayAsync(MoveTemp(Action), Ray.TransformBy(GetComponentTransform().Inverse()), Radius, bVisibleOnly, MoveTemp(CompletionCallback));
+		}
+	}
+
+	/** Applies the given color to all points within the sphere */
+	UFUNCTION(BlueprintCallable, Category = "Lidar Point Cloud")
+	void ApplyColorToPointsInSphere(FColor NewColor, FVector Center, float Radius, bool bVisibleOnly) { ApplyColorToPointsInSphere(NewColor, FSphere(Center, Radius), bVisibleOnly); }
+	void ApplyColorToPointsInSphere(const FColor& NewColor, const FSphere& Sphere, const bool& bVisibleOnly)
+	{
+		if (PointCloud)
+		{
+			PointCloud->ApplyColorToPointsInSphere(NewColor, Sphere.TransformBy(GetComponentTransform().Inverse()), bVisibleOnly);
+		}
+	}
+
+	/** Applies the given color to all points within the box */
+	UFUNCTION(BlueprintCallable, Category = "Lidar Point Cloud")
+	void ApplyColorToPointsInBox(FColor NewColor, FVector Center, FVector Extent, bool bVisibleOnly) { ApplyColorToPointsInBox(NewColor, FBox(Center - Extent, Center + Extent), bVisibleOnly); }
+	void ApplyColorToPointsInBox(const FColor& NewColor, const FBox& Box, const bool& bVisibleOnly)
+	{
+		if (PointCloud)
+		{
+			PointCloud->ApplyColorToPointsInBox(NewColor, Box.TransformBy(GetComponentTransform().Inverse()), bVisibleOnly);
+		}
+	}
+
+	/** Applies the given color to all points hit by the given ray */
+	UFUNCTION(BlueprintCallable, Category = "Lidar Point Cloud")
+	void ApplyColorToPointsByRay(FColor NewColor, FVector Origin, FVector Direction, float Radius, bool bVisibleOnly) { ApplyColorToPointsByRay(NewColor, FLidarPointCloudRay(Origin, Direction), Radius, bVisibleOnly); }
+	void ApplyColorToPointsByRay(const FColor& NewColor, FLidarPointCloudRay Ray, float Radius, bool bVisibleOnly)
+	{
+		if (PointCloud)
+		{
+			PointCloud->ApplyColorToPointsByRay(NewColor, Ray.TransformBy(GetComponentTransform().Inverse()), Radius, bVisibleOnly);
+		}
+	}
+	
+	/**
+	 * Removes all points within the given sphere
+	 * If bRefreshPointsBounds is set to false, make sure you call RefreshPointsBounds() manually or cloud centering may not work correctly.
+	 */
+	UFUNCTION(BlueprintCallable, Category = "Lidar Point Cloud")
+	void RemovePointsInSphere(FVector Center, float Radius, bool bVisibleOnly, bool bRefreshPointsBounds) { RemovePointsInSphere(FSphere(Center, Radius), bVisibleOnly, bRefreshPointsBounds); }
+	void RemovePointsInSphere(const FSphere& Sphere, const bool& bVisibleOnly, bool bRefreshPointsBounds)
+	{
+		if (PointCloud)
+		{
+			PointCloud->RemovePointsInSphere(Sphere.TransformBy(GetComponentTransform().Inverse()), bVisibleOnly, bRefreshPointsBounds);
+		}
+	}
+
+	/**
+	 * Removes all points within the given box
+	 * If bRefreshPointsBounds is set to false, make sure you call RefreshPointsBounds() manually or cloud centering may not work correctly.
+	 */
+	UFUNCTION(BlueprintCallable, Category = "Lidar Point Cloud")
+	void RemovePointsInBox(FVector Center, FVector Extent, bool bVisibleOnly, bool bRefreshPointsBounds) { RemovePointsInBox(FBox(Center - Extent, Center + Extent), bVisibleOnly, bRefreshPointsBounds); }
+	void RemovePointsInBox(const FBox& Box, const bool& bVisibleOnly, bool bRefreshPointsBounds)
+	{
+		if (PointCloud)
+		{
+			PointCloud->RemovePointsInBox(Box.TransformBy(GetComponentTransform().Inverse()), bVisibleOnly, bRefreshPointsBounds);
+		}
+	}
+
+	/**
+	 * Removes all points hit by the given ray
+	 * If bRefreshPointsBounds is set to false, make sure you call RefreshPointsBounds() manually or cloud centering may not work correctly.
+	 */
+	UFUNCTION(BlueprintCallable, Category = "Lidar Point Cloud")
+	void RemovePointsByRay(FVector Origin, FVector Direction, float Radius, bool bVisibleOnly, bool bRefreshPointsBounds) { RemovePointsByRay(FLidarPointCloudRay(Origin, Direction), Radius, bVisibleOnly, bRefreshPointsBounds); }
+	void RemovePointsByRay(FLidarPointCloudRay Ray, float Radius, const bool& bVisibleOnly, bool bRefreshPointsBounds)
+	{
+		if (PointCloud)
+		{
+			PointCloud->RemovePointsByRay(Ray.TransformBy(GetComponentTransform().Inverse()), Radius, bVisibleOnly, bRefreshPointsBounds);
+		}
+	}
+
+public:
+	UFUNCTION(BlueprintCallable, Category = "Components|LidarPointCloud")
+	void SetPointCloud(ULidarPointCloud *InPointCloud);
+
+	/** Applies specified rendering parameters (Brightness, Saturation, etc) to the selected material */
+	UFUNCTION(BlueprintCallable, Category = "Components|LidarPointCloud|Rendering")
+	void ApplyRenderingParameters();
+
+public:
+	// Begin UObject Interface.
+	static void AddReferencedObjects(UObject* InThis, FReferenceCollector& Collector);
+	virtual void PostLoad() override;
+	// End UObject Interface.
+
+	// End UMeshComponent Interface
+	virtual int32 GetNumMaterials() const override { return 1; }
+	virtual UMaterialInterface* GetMaterial(int32 ElementIndex) const override { return Material; }
+	virtual void SetMaterial(int32 ElementIndex, UMaterialInterface* InMaterial) override;
+	// End UMeshComponent Interface
+
+#if WITH_EDITOR
+	virtual void PreEditChange(FProperty* PropertyThatWillChange) override;
+	virtual void PostEditChangeProperty(FPropertyChangedEvent& PropertyChangedEvent) override;
+	virtual void PostEditImport() override
+	{
+		Super::PostEditImport();
+
+		// Make sure to update the material after duplicating this component
+		UpdateMaterial();
+	}
+#endif
+
+	virtual UBodySetup* GetBodySetup() override;
+
+private:
+	// Begin UPrimitiveComponent Interface
+	virtual FPrimitiveSceneProxy* CreateSceneProxy() override;
+	// End UMeshComponent Interface
+
+	// Begin USceneComponent Interface
+	virtual FBoxSphereBounds CalcBounds(const FTransform& LocalToWorld) const override;
+	// End USceneComponent Interface
+
+	void UpdateMaterial();
+
+	void AttachPointCloudListener();
+	void RemovePointCloudListener();
+	void OnPointCloudRebuilt();
+	void OnPointCloudCollisionUpdated();
+
+	void PostPointCloudSet();
+
+	friend class FPointCloudSceneProxy;
+#if WITH_EDITOR
+	friend class FLidarPointCloudEditorViewportClient;
+#endif
+};
