Fixes for the static analysis.

diff --git a/Engine/Plugins/Enterprise/LidarPointCloud/Source/LidarPointCloudRuntime/Private/LidarPointCloud.cpp b/Engine/Plugins/Enterprise/LidarPointCloud/Source/LidarPointCloudRuntime/Private/LidarPointCloud.cpp
--- a/Engine/Plugins/Enterprise/LidarPointCloud/Source/LidarPointCloudRuntime/Private/LidarPointCloud.cpp
+++ b/Engine/Plugins/Enterprise/LidarPointCloud/Source/LidarPointCloudRuntime/Private/LidarPointCloud.cpp
@@ -582,7 +582,8 @@ bool ULidarPointCloud::InsertPoints_NoLock(T InPoints, const int64& Count, ELida
 		int64 TotalProcessedPoints = 0;
 
 		const int32 NumThreads = FMath::Min(FPlatformMisc::NumberOfCoresIncludingHyperthreads() - 1, (int32)(Count / MaxBatchSize) + 1);
-		TFuture<void>* ThreadResults = new TFuture<void>[NumThreads];
+		TArray<TFuture<void>>ThreadResults;
+		ThreadResults.Reserve( NumThreads );
 		const int64 NumPointsPerThread = Count / NumThreads + 1;
 
 		FCriticalSection ProgressCallbackLock;
@@ -591,8 +592,7 @@ bool ULidarPointCloud::InsertPoints_NoLock(T InPoints, const int64& Count, ELida
 		for (int32 t = 0; t < NumThreads; t++)
 		{
 			const int32 ThreadID = t;
-
-			ThreadResults[ThreadID] = Async(EAsyncExecution::Thread, [this, ThreadID, DuplicateHandling, bRefreshPointsBounds, MaxBatchSize, NumPointsPerThread, RefreshStatusFrequency, &ProcessedPoints, &TotalProcessedPoints, InPoints, Count, &ProgressCallback, &ProgressCallbackLock, &bCanceled]
+			ThreadResults.Add(Async(EAsyncExecution::Thread, [this, ThreadID, DuplicateHandling, bRefreshPointsBounds, MaxBatchSize, NumPointsPerThread, RefreshStatusFrequency, &ProcessedPoints, &TotalProcessedPoints, InPoints, Count, &ProgressCallback, &ProgressCallbackLock, &bCanceled]
 				{
 					int64 Idx = ThreadID * NumPointsPerThread;
 					int64 MaxIdx = FMath::Min(Idx + NumPointsPerThread, Count);
@@ -624,14 +624,14 @@ bool ULidarPointCloud::InsertPoints_NoLock(T InPoints, const int64& Count, ELida
 						Idx += BatchSize;
 						DataPointer += BatchSize;
 					}
-				});
+				}));
 		}
 
 		// Sync
-		for (int32 i = 0; i < NumThreads; i++) ThreadResults[i].Get();
-
-		// Cleanup
-		delete[] ThreadResults;
+		for (auto& ThreadResult : ThreadResults)
+		{
+			ThreadResult.Get();
+		}
 	}
 	else
 	{
@@ -1100,6 +1100,7 @@ bool ULidarPointCloudBlueprintLibrary::LineTraceMulti(UObject* WorldContextObjec
 	return Hits.Num() > 0;
 }
 
+template bool ULidarPointCloud::InsertPoints_NoLock<FLidarPointCloudPoint*>(FLidarPointCloudPoint*, const int64&, ELidarPointCloudDuplicateHandling, bool, FThreadSafeBool*, TFunction<void(float)>);
 template bool ULidarPointCloud::InsertPoints_NoLock<const FLidarPointCloudPoint*>(const FLidarPointCloudPoint*, const int64&, ELidarPointCloudDuplicateHandling, bool, FThreadSafeBool*, TFunction<void(float)>);
 template bool ULidarPointCloud::InsertPoints_NoLock<FLidarPointCloudPoint**>(FLidarPointCloudPoint**, const int64&, ELidarPointCloudDuplicateHandling, bool, FThreadSafeBool*, TFunction<void(float)>);
 template bool ULidarPointCloud::SetData<const FLidarPointCloudPoint*>(const FLidarPointCloudPoint*, const int64&, TFunction<void(float)>);
diff --git a/Engine/Plugins/Enterprise/LidarPointCloud/Source/LidarPointCloudRuntime/Private/LidarPointCloudOctree.cpp b/Engine/Plugins/Enterprise/LidarPointCloud/Source/LidarPointCloudRuntime/Private/LidarPointCloudOctree.cpp
--- a/Engine/Plugins/Enterprise/LidarPointCloud/Source/LidarPointCloudRuntime/Private/LidarPointCloudOctree.cpp
+++ b/Engine/Plugins/Enterprise/LidarPointCloud/Source/LidarPointCloudRuntime/Private/LidarPointCloudOctree.cpp
@@ -1699,30 +1699,19 @@ FArchive& operator<<(FArchive& Ar, FLidarPointCloudOctree& O)
 FLidarPointCloudTraversalOctreeNode::FLidarPointCloudTraversalOctreeNode()
 	: DataNode(nullptr)
 	, Parent(nullptr)
-	, Children(nullptr)
 {
 }
 
-FLidarPointCloudTraversalOctreeNode::~FLidarPointCloudTraversalOctreeNode()
-{
-	if (Children)
-	{
-		delete[] Children;
-	}
-}
-
 void FLidarPointCloudTraversalOctreeNode::Build(FLidarPointCloudOctreeNode* Node, const FTransform& LocalToWorld)
 {
 	DataNode = Node;
 	Center = LocalToWorld.TransformPosition(Node->Center);
 	Depth = Node->Depth;
 
-	NumChildren = Node->Children.Num();
-	if (NumChildren > 0)
+	Children.AddZeroed(Node->Children.Num());
+	for (int32 i = 0; i < Children.Num(); i++)
 	{
-		Children = new FLidarPointCloudTraversalOctreeNode[NumChildren];
-
-		for (int32 i = 0; i < NumChildren; i++)
+		if (Node->Children[i])
 		{
 			Children[i].Build(Node->Children[i], LocalToWorld);
 			Children[i].Parent = this;
@@ -1745,11 +1734,11 @@ void FLidarPointCloudTraversalOctreeNode::CalculateVirtualDepth(const TArray<flo
 	Nodes.Enqueue(this);
 	while (Nodes.Dequeue(CurrentNode))
 	{
-		for (int32 i = 0; i < CurrentNode->NumChildren; i++)
+		for (auto& Child : CurrentNode->Children)
 		{
-			if (CurrentNode->Children[i].bSelected)
+			if (Child.bSelected)
 			{
-				Nodes.Enqueue(&CurrentNode->Children[i]);
+				Nodes.Enqueue(&Child);
 			}
 		}
 
@@ -1757,7 +1746,7 @@ void FLidarPointCloudTraversalOctreeNode::CalculateVirtualDepth(const TArray<flo
 
 		if (CurrentNode != this && PointSizeBias > 0)
 		{
-			LocalVDFactor /= (CurrentNode->Parent->NumChildren - 1) * PointSizeBias + 1;
+			LocalVDFactor /= (CurrentNode->Parent->Children.Num() - 1) * PointSizeBias + 1;
 		}
 
 		VDFactor += LocalVDFactor;
@@ -1768,11 +1757,11 @@ void FLidarPointCloudTraversalOctreeNode::CalculateVirtualDepth(const TArray<flo
 	Nodes.Enqueue(this);
 	while (Nodes.Dequeue(CurrentNode))
 	{
-		for (int32 i = 0; i < CurrentNode->NumChildren; i++)
+		for (auto& Child : CurrentNode->Children)
 		{
-			if (CurrentNode->Children[i].bSelected)
+			if (Child.bSelected)
 			{
-				Nodes.Enqueue(&CurrentNode->Children[i]);
+				Nodes.Enqueue(&Child);
 			}
 		}
 
@@ -1897,9 +1886,9 @@ TArray<FLidarPointCloudTraversalOctreeNode*> FLidarPointCloudTraversalOctree::Ge
 
 			if (MaxDepth < 0 || CurrentNode->Depth < MaxDepth)
 			{
-				for (int32 i = 0; i < CurrentNode->NumChildren; i++)
+				for (auto& Child : CurrentNode->Children)
 				{
-					Nodes.Enqueue(&CurrentNode->Children[i]);
+					Nodes.Enqueue(&Child);
 				}
 			}
 		}
diff --git a/Engine/Plugins/Enterprise/LidarPointCloud/Source/LidarPointCloudRuntime/Public/LidarPointCloudOctree.h b/Engine/Plugins/Enterprise/LidarPointCloud/Source/LidarPointCloudRuntime/Public/LidarPointCloudOctree.h
--- a/Engine/Plugins/Enterprise/LidarPointCloud/Source/LidarPointCloudRuntime/Public/LidarPointCloudOctree.h
+++ b/Engine/Plugins/Enterprise/LidarPointCloud/Source/LidarPointCloudRuntime/Public/LidarPointCloudOctree.h
@@ -461,18 +461,12 @@ struct FLidarPointCloudTraversalOctreeNode
 	FLidarPointCloudTraversalOctreeNode* Parent;
 
 	/** Stores the children array */
-	FLidarPointCloudTraversalOctreeNode* Children;
-	uint8 NumChildren : 4;
+	TArray<FLidarPointCloudTraversalOctreeNode> Children;
 
 	/** Holds true if the node has been selected for rendering. */
-	uint8 bSelected : 4;
+	bool bSelected;
 
 	FLidarPointCloudTraversalOctreeNode();
-	~FLidarPointCloudTraversalOctreeNode();
-	FLidarPointCloudTraversalOctreeNode(const FLidarPointCloudTraversalOctreeNode&) = delete;
-	FLidarPointCloudTraversalOctreeNode(FLidarPointCloudTraversalOctreeNode&&) = delete;
-	FLidarPointCloudTraversalOctreeNode& operator=(const FLidarPointCloudTraversalOctreeNode&) = delete;
-	FLidarPointCloudTraversalOctreeNode& operator=(FLidarPointCloudTraversalOctreeNode&&) = delete;
 
 	/** Builds the traversal version of the given node. */
 	void Build(FLidarPointCloudOctreeNode* Node, const FTransform& LocalToWorld);
