Merging
//UE4/Dev-Enterprise/Engine/Plugins/Enterprise/LidarPointCloud/...
to //UE4/Release-4.25/Engine/Plugins/Enterprise/LidarPointCloud/...

Also adding the shelve from Michal to upgrade the plugin to the new rendering api.

diff --git a/Engine/Plugins/Enterprise/LidarPointCloud/Shaders/Private/LidarPointCloudVertexFactory.ush b/Engine/Plugins/Enterprise/LidarPointCloud/Shaders/Private/LidarPointCloudVertexFactory.ush
--- a/Engine/Plugins/Enterprise/LidarPointCloud/Shaders/Private/LidarPointCloudVertexFactory.ush
+++ b/Engine/Plugins/Enterprise/LidarPointCloud/Shaders/Private/LidarPointCloudVertexFactory.ush
@@ -18,6 +18,10 @@ static float3 LODColors[] = {
 
 // Used for instance handling
 uint IndexDivisor;
+uint FirstElementIndex;
+
+// Alignment
+float3 LocationOffset;
 
 // Used for Sprite Size
 float VDMultiplier;
@@ -31,7 +35,6 @@ float3 ViewUpVector;
 
 // Used for coloration override
 float3 BoundsSize;
-float3 BoundsOffset;
 half3 ElevationColorBottom;
 half3 ElevationColorTop;
 
@@ -52,7 +55,7 @@ uint bUseClassification;
 
 float4 ClassificationColors[32];
 
-StructuredBuffer<float> DataBuffer;
+Buffer<float> DataBuffer;
 
 #if USE_INSTANCING
 	// Instanced stereo sets the eye index explicitly for instanced geometry
@@ -114,10 +117,10 @@ FPointCloudInstanceData GetInstanceData(uint VertexId)
 {
 	FPointCloudInstanceData InstanceData = (FPointCloudInstanceData)0;
 
-	uint PointId = VertexId / IndexDivisor;
+	uint PointId = FirstElementIndex + VertexId / IndexDivisor;
 	uint Idx = PointId * 4;
 
-	InstanceData.Position = float3(DataBuffer[Idx], DataBuffer[Idx + 1], DataBuffer[Idx + 2]);
+	InstanceData.Position = float3(DataBuffer[Idx], DataBuffer[Idx + 1], DataBuffer[Idx + 2]) + LocationOffset;
 	InstanceData.Color = asuint(DataBuffer[Idx + 3]);
 	InstanceData.VirtualDepth = half((0x000000FF & (asuint(DataBuffer[SizeOffset + PointId / 4]) >> (PointId % 4 * 8)))) * VDMultiplier;
 	InstanceData.Scale = RootCellSize * PrecomputedUV_Quad[VertexId % 4] / pow(2, InstanceData.VirtualDepth);
@@ -152,11 +155,11 @@ FVertexFactoryIntermediates GetVertexFactoryIntermediates(FVertexFactoryInput In
 	// Color
 	{
 		// Used to account for offset due to shifted points bounds
-		float3 OffsetPosition = Intermediates.Position + BoundsOffset;
+		float3 RelativePosition = Intermediates.Position - LocationOffset;
 
 		half3 Color = half3(0x000000FF & (InstanceData.Color >> 16), 0x000000FF & (InstanceData.Color >> 8), 0x000000FF & InstanceData.Color) * 0.003921568627451;
-		half3 ElevationColor = lerp(ElevationColorBottom, ElevationColorTop, clamp((OffsetPosition.z + BoundsSize.z * 0.5) / BoundsSize.z, 0, 1));
-		half3 PositionColor = (OffsetPosition + BoundsSize * 0.5) / BoundsSize;
+		half3 ElevationColor = lerp(ElevationColorBottom, ElevationColorTop, clamp((RelativePosition.z + BoundsSize.z * 0.5) / BoundsSize.z, 0, 1));
+		half3 PositionColor = (RelativePosition + BoundsSize * 0.5) / BoundsSize;
 		half3 ClassificationColor = ClassificationColors[0x0000001F & InstanceData.Color];
 
 		// Pick color source
diff --git a/Engine/Plugins/Enterprise/LidarPointCloud/Source/LidarPointCloudEditor/Private/LidarPointCloudEditor.cpp b/Engine/Plugins/Enterprise/LidarPointCloud/Source/LidarPointCloudEditor/Private/LidarPointCloudEditor.cpp
--- a/Engine/Plugins/Enterprise/LidarPointCloud/Source/LidarPointCloudEditor/Private/LidarPointCloudEditor.cpp
+++ b/Engine/Plugins/Enterprise/LidarPointCloud/Source/LidarPointCloudEditor/Private/LidarPointCloudEditor.cpp
@@ -81,6 +81,7 @@ FLidarPointCloudEditor::~FLidarPointCloudEditor()
 	if (PointCloudBeingEdited)
 	{
 		PointCloudBeingEdited->OnPointCloudRebuilt().RemoveAll(this);
+		PointCloudBeingEdited->OnPreSaveCleanup().RemoveAll(this);
 	}
 
 	DeselectPoints();
@@ -144,6 +145,7 @@ void FLidarPointCloudEditor::InitPointCloudEditor(const EToolkitMode::Type Mode,
 
 	// Register for rebuilding events
 	PointCloudBeingEdited->OnPointCloudRebuilt().AddSP(this, &FLidarPointCloudEditor::OnPointCloudRebuilt);
+	PointCloudBeingEdited->OnPreSaveCleanup().AddSP(this, &FLidarPointCloudEditor::OnPreSaveCleanup);
 
 	TSharedPtr<FLidarPointCloudEditor> PointCloudEditor = SharedThis(this);
 
@@ -262,13 +264,18 @@ void FLidarPointCloudEditor::DeselectPoints()
 
 void FLidarPointCloudEditor::DeletePoints()
 {
+	if (!ConfirmCollisionChange())
+	{
+		return;
+	}
+
 	// Permanently remove points
 	if (FMessageDialog::Open(EAppMsgType::Type::YesNo, FText::FromString("Warning: This operation cannot be reversed!\nAre you sure you want to continue?")) != EAppReturnType::Yes)
 	{
 		return;
 	}
 
-	PointCloudBeingEdited->RemovePoints_NoLock(SelectedPoints, true);
+	PointCloudBeingEdited->RemovePoints(SelectedPoints);
 	PointCloudBeingEdited->RefreshRendering();
 	PointCloudBeingEdited->MarkPackageDirty();
 	SelectedPoints.Empty();
@@ -284,7 +291,7 @@ void FLidarPointCloudEditor::DeleteHiddenPoints()
 		return;
 	}
 
-	PointCloudBeingEdited->RemoveHiddenPoints(true);
+	PointCloudBeingEdited->RemoveHiddenPoints();
 	DeselectPoints();
 	PointCloudBeingEdited->RefreshRendering();
 	PointCloudBeingEdited->MarkPackageDirty();
@@ -294,6 +301,11 @@ void FLidarPointCloudEditor::DeleteHiddenPoints()
 
 void FLidarPointCloudEditor::HidePoints()
 {
+	if (!ConfirmCollisionChange())
+	{
+		return;
+	}
+
 	// Hide points
 	for (FLidarPointCloudPoint** Data = SelectedPoints.GetData(), **DataEnd = Data + SelectedPoints.Num(); Data != DataEnd; ++Data)
 	{
@@ -310,10 +322,30 @@ void FLidarPointCloudEditor::UnhideAll()
 {
 	if (PointCloudBeingEdited)
 	{
+		if (!ConfirmCollisionChange())
+		{
+			return;
+		}
+
 		PointCloudBeingEdited->UnhideAll();
 	}
 }
 
+bool FLidarPointCloudEditor::ConfirmCollisionChange()
+{
+	if (PointCloudBeingEdited->HasCollisionData())
+	{
+		if (FMessageDialog::Open(EAppMsgType::Type::YesNo, FText::FromString("Performing this action will invalidate the collision data.\nAre you sure you want to continue?")) != EAppReturnType::Yes)
+		{
+			return false;
+		}
+
+		PointCloudBeingEdited->RemoveCollision();
+	}
+
+	return true;
+}
+
 TSharedRef<SWidget> FLidarPointCloudEditor::BuildPointCloudStatistics()
 {
 	return SNew(SHorizontalBox)
@@ -355,7 +387,7 @@ TSharedRef<SWidget> FLidarPointCloudEditor::BuildPointCloudStatistics()
 			[
 				SNew(STextBlock).Text_Lambda([this]
 				{
-					FVector BoundingSize = PointCloudBeingEdited ? PointCloudBeingEdited->GetPointsBounds().GetSize() : FVector::ZeroVector;
+					FVector BoundingSize = PointCloudBeingEdited ? PointCloudBeingEdited->GetBounds().GetSize() : FVector::ZeroVector;
 					return FText::Format(LOCTEXT("PCBounds", "Bounds: {0} x {1} x {2}"), FMath::CeilToInt(BoundingSize.X), FMath::CeilToInt(BoundingSize.Y), FMath::CeilToInt(BoundingSize.Z));
 				})
 			]
@@ -545,6 +577,12 @@ void FLidarPointCloudEditor::OnPointCloudRebuilt()
 	SelectedPoints.Empty();
 }
 
+void FLidarPointCloudEditor::OnPreSaveCleanup()
+{
+	// If the cloud asset is being saved, deselect all points
+	DeselectPoints();
+}
+
 void FLidarPointCloudEditor::Extract()
 {
 	// Skip, if no points are selected
@@ -553,6 +591,11 @@ void FLidarPointCloudEditor::Extract()
 		return;
 	}
 
+	if (!ConfirmCollisionChange())
+	{
+		return;
+	}
+
 	ULidarPointCloud* NewPointCloud = CreateNewAsset();
 	if (NewPointCloud)
 	{
@@ -581,6 +624,13 @@ void FLidarPointCloudEditor::ExtractCopy()
 
 void FLidarPointCloudEditor::ToggleCenter()
 {
+	if (!ConfirmCollisionChange())
+	{
+		return;
+	}
+
+	DeselectPoints();
+
 	if (IsCentered())
 	{
 		PointCloudBeingEdited->RestoreOriginalCoordinates();
@@ -598,6 +648,20 @@ void FLidarPointCloudEditor::ToggleCenter()
 
 void FLidarPointCloudEditor::ToggleEditMode()
 {
+	if (!bEditMode)
+	{
+		// Warn about loading the whole cloud
+		if (!PointCloudBeingEdited->IsFullyLoaded())
+		{
+			if (FMessageDialog::Open(EAppMsgType::Type::YesNo, FText::FromString("The asset needs to be fully loaded into memory to enable editing.\nThis may take a while, depending on the size of the asset.\nAre you sure you want to continue?")) != EAppReturnType::Yes)
+			{
+				return;
+			}
+		}
+
+		PointCloudBeingEdited->LoadAllNodes();
+	}
+
 	bEditMode = !bEditMode;
 
 	if (!bEditMode)
@@ -662,6 +726,11 @@ void FLidarPointCloudEditor::RemoveCollision()
 
 void FLidarPointCloudEditor::Align()
 {
+	if (!ConfirmCollisionChange())
+	{
+		return;
+	}
+
 	TArray<FAssetData> SelectedAssets = SelectAssets(LOCTEXT("SelectAlignmentSources", "Select Alignment Sources"));
 	if (SelectedAssets.Num())
 	{
diff --git a/Engine/Plugins/Enterprise/LidarPointCloud/Source/LidarPointCloudEditor/Private/LidarPointCloudEditor.h b/Engine/Plugins/Enterprise/LidarPointCloud/Source/LidarPointCloudEditor/Private/LidarPointCloudEditor.h
--- a/Engine/Plugins/Enterprise/LidarPointCloud/Source/LidarPointCloudEditor/Private/LidarPointCloudEditor.h
+++ b/Engine/Plugins/Enterprise/LidarPointCloud/Source/LidarPointCloudEditor/Private/LidarPointCloudEditor.h
@@ -69,6 +69,8 @@ public:
 	bool IsEditMode() const { return bEditMode; }
 
 private:
+	bool ConfirmCollisionChange();
+
 	TSharedRef<SWidget> BuildPointCloudStatistics();
 
 	/** Builds the Point Cloud Editor toolbar. */
@@ -82,6 +84,8 @@ private:
 	/** Called when the owned cloud is rebuilt */
 	void OnPointCloudRebuilt();
 
+	void OnPreSaveCleanup();
+
 	void Extract();
 	void ExtractCopy();
 	bool IsCentered() const;
diff --git a/Engine/Plugins/Enterprise/LidarPointCloud/Source/LidarPointCloudEditor/Private/LidarPointCloudEditorViewport.cpp b/Engine/Plugins/Enterprise/LidarPointCloud/Source/LidarPointCloudEditor/Private/LidarPointCloudEditorViewport.cpp
--- a/Engine/Plugins/Enterprise/LidarPointCloud/Source/LidarPointCloudEditor/Private/LidarPointCloudEditorViewport.cpp
+++ b/Engine/Plugins/Enterprise/LidarPointCloud/Source/LidarPointCloudEditor/Private/LidarPointCloudEditorViewport.cpp
@@ -54,8 +54,8 @@ public:
 
 void SLidarPointCloudEditorViewport::Construct(const FArguments& InArgs)
 {
-	PreviewScene->SetFloorVisibility(false, true);
-	PreviewScene->SetEnvironmentVisibility(false);
+	PreviewScene.SetFloorVisibility(false, true);
+	PreviewScene.SetEnvironmentVisibility(false);
 
 	PointCloudEditorPtr = InArgs._PointCloudEditor;
 
@@ -63,9 +63,13 @@ void SLidarPointCloudEditorViewport::Construct(const FArguments& InArgs)
 
 	CurrentViewMode = VMI_Lit;
 
+	PreviewCloudComponent = NewObject<ULidarPointCloudComponent>();
+
 	SEditorViewport::Construct(SEditorViewport::FArguments());
 
-	PreviewCloudComponent = NewObject<ULidarPointCloudComponent>((UObject*)GetTransientPackage(), NAME_None, RF_Transient);
+	FComponentReregisterContext ReregisterContext(PreviewCloudComponent);
+
+	PreviewScene.AddComponent(PreviewCloudComponent, FTransform::Identity);
 
 	SetPreviewCloud(PointCloud);
 
@@ -81,7 +85,7 @@ void SLidarPointCloudEditorViewport::Construct(const FArguments& InArgs)
 }
 
 SLidarPointCloudEditorViewport::SLidarPointCloudEditorViewport()
-	: PreviewScene(MakeShareable(new FAdvancedPreviewScene(FPreviewScene::ConstructionValues(), FLT_MAX)))
+	: PreviewScene(FPreviewScene::ConstructionValues(), FLT_MAX)
 {
 }
 
@@ -92,6 +96,13 @@ SLidarPointCloudEditorViewport::~SLidarPointCloudEditorViewport()
 	{
 		EditorViewportClient->Viewport = nullptr;
 	}
+
+	if (PreviewCloudComponent)
+	{
+		PreviewScene.RemoveComponent(PreviewCloudComponent);
+		PreviewCloudComponent->SetPointCloud(nullptr);
+		PreviewCloudComponent = nullptr;
+	}
 }
 
 void SLidarPointCloudEditorViewport::AddReferencedObjects(FReferenceCollector& Collector)
@@ -117,25 +128,11 @@ void SLidarPointCloudEditorViewport::ResetCamera()
 void SLidarPointCloudEditorViewport::SetPreviewCloud(ULidarPointCloud* InPointCloud)
 {
 	// Set the new preview point cloud.
-	FComponentReregisterContext ReregisterContext(PreviewCloudComponent);
 	PreviewCloudComponent->SetPointCloud(InPointCloud);
+	PreviewCloudComponent->OwningViewportClient = TWeakPtr<FViewportClient>(EditorViewportClient);
+	PreviewCloudComponent->MarkRenderStateDirty();
 
-	PreviewScene->AddComponent(PreviewCloudComponent, FTransform::Identity);
-
-	EditorViewportClient->SetPreviewCloud(InPointCloud, PreviewCloudComponent);
-}
-
-void SLidarPointCloudEditorViewport::UpdatePreviewCloud(ULidarPointCloud* InPointCloud, bool bResetCamera/*= true*/)
-{
-	if (PreviewCloudComponent)
-	{
-		PreviewScene->RemoveComponent(PreviewCloudComponent);
-		PreviewCloudComponent = nullptr;
-	}
-
-	PreviewCloudComponent = NewObject<ULidarPointCloudComponent>();
-
-	SetPreviewCloud(InPointCloud);
+	ResetCamera();
 }
 
 void SLidarPointCloudEditorViewport::PopulateOverlayText(const TArray<FOverlayTextItem>& TextItems)
@@ -172,7 +169,7 @@ void SLidarPointCloudEditorViewport::OnObjectPropertyChanged(UObject* ObjectBein
 
 TSharedRef<FEditorViewportClient> SLidarPointCloudEditorViewport::MakeEditorViewportClient()
 {
-	EditorViewportClient = MakeShareable(new FLidarPointCloudEditorViewportClient(PointCloudEditorPtr, SharedThis(this), PreviewScene.ToSharedRef(), PointCloud, nullptr));
+	EditorViewportClient = MakeShareable(new FLidarPointCloudEditorViewportClient(PointCloudEditorPtr, SharedThis(this), &PreviewScene, PointCloud, PreviewCloudComponent));
 
 	EditorViewportClient->bSetListenerPosition = false;
 
diff --git a/Engine/Plugins/Enterprise/LidarPointCloud/Source/LidarPointCloudEditor/Private/LidarPointCloudEditorViewport.h b/Engine/Plugins/Enterprise/LidarPointCloud/Source/LidarPointCloudEditor/Private/LidarPointCloudEditorViewport.h
--- a/Engine/Plugins/Enterprise/LidarPointCloud/Source/LidarPointCloudEditor/Private/LidarPointCloudEditorViewport.h
+++ b/Engine/Plugins/Enterprise/LidarPointCloud/Source/LidarPointCloudEditor/Private/LidarPointCloudEditorViewport.h
@@ -42,9 +42,6 @@ public:
 	/** Component for the preview point cloud. */
 	ULidarPointCloudComponent* PreviewCloudComponent;
 
-	/** Retrieves the static mesh component. */
-	ULidarPointCloudComponent* GetPointCloudComponent() const { return PreviewCloudComponent; }
-
 	/**
 	 *	Sets up the point cloud that the Point Cloud editor is viewing.
 	 *
@@ -52,16 +49,6 @@ public:
 	 */
 	void SetPreviewCloud(ULidarPointCloud* InPointCloud);
 
-	/**
-	 *	Updates the preview cloud and other viewport specific settings that go with it.
-	 *
-	 *	@param	InPointCloud		The point cloud being viewed in the editor.
-	 */
-	void UpdatePreviewCloud(ULidarPointCloud* InPointCloud, bool bResetCamera = true);
-
-	/** @return The editor viewport client */
-	FLidarPointCloudEditorViewportClient& GetViewportClient() { return *EditorViewportClient; }
-
 	/** Set the parent tab of the viewport for determining visibility */
 	void SetParentTab(TSharedRef<SDockTab> InParentTab) { ParentTab = InParentTab; }
 
@@ -85,9 +72,6 @@ public:
 	virtual void OnFloatingButtonClicked() override {}
 	// End of ICommonEditorViewportToolbarInfoProvider interface
 
-	/** Returns the preview scene being rendered in the viewport */
-	TSharedRef<FAdvancedPreviewScene> GetPreviewScene() { return PreviewScene.ToSharedRef(); }
-
 protected:
 	/** SEditorViewport interface */
 	virtual TSharedRef<FEditorViewportClient> MakeEditorViewportClient() override;
@@ -111,7 +95,7 @@ private:
 	TWeakPtr<FLidarPointCloudEditor> PointCloudEditorPtr;
 
 	/** The scene for this viewport. */
-	TSharedPtr<FAdvancedPreviewScene> PreviewScene;
+	FAdvancedPreviewScene PreviewScene;
 
 	/** Editor viewport client */
 	TSharedPtr<FLidarPointCloudEditorViewportClient> EditorViewportClient;
diff --git a/Engine/Plugins/Enterprise/LidarPointCloud/Source/LidarPointCloudEditor/Private/LidarPointCloudEditorViewportClient.cpp b/Engine/Plugins/Enterprise/LidarPointCloud/Source/LidarPointCloudEditor/Private/LidarPointCloudEditorViewportClient.cpp
--- a/Engine/Plugins/Enterprise/LidarPointCloud/Source/LidarPointCloudEditor/Private/LidarPointCloudEditorViewportClient.cpp
+++ b/Engine/Plugins/Enterprise/LidarPointCloud/Source/LidarPointCloudEditor/Private/LidarPointCloudEditorViewportClient.cpp
@@ -27,8 +27,9 @@ namespace {
 	static float AmbientCubemapIntensity = 0.4f;
 }
 
-FLidarPointCloudEditorViewportClient::FLidarPointCloudEditorViewportClient(TWeakPtr<FLidarPointCloudEditor> InPointCloudEditor, const TSharedRef<SLidarPointCloudEditorViewport>& InPointCloudEditorViewport, const TSharedRef<FAdvancedPreviewScene>& InPreviewScene, ULidarPointCloud* InPreviewPointCLoud, ULidarPointCloudComponent* InPreviewPointCloudComponent)
-	: FEditorViewportClient(nullptr, &InPreviewScene.Get(), StaticCastSharedRef<SEditorViewport>(InPointCloudEditorViewport))
+FLidarPointCloudEditorViewportClient::FLidarPointCloudEditorViewportClient(TWeakPtr<FLidarPointCloudEditor> InPointCloudEditor, const TSharedRef<SLidarPointCloudEditorViewport>& InPointCloudEditorViewport, FAdvancedPreviewScene* InPreviewScene, ULidarPointCloud* InPreviewPointCloud, ULidarPointCloudComponent* InPreviewPointCloudComponent)
+	: FEditorViewportClient(nullptr, InPreviewScene, StaticCastSharedRef<SEditorViewport>(InPointCloudEditorViewport))
+	, PointCloudComponent(InPreviewPointCloudComponent)
 	, PointCloudEditorPtr(InPointCloudEditor)
 	, PointCloudEditorViewportPtr(InPointCloudEditorViewport)
 	, SelectionMode(ELidarPointCloudSelectionMode::None)
@@ -44,7 +45,7 @@ FLidarPointCloudEditorViewportClient::FLidarPointCloudEditorViewportClient(TWeak
 	DrawHelper.PerspectiveGridSize = GridSize;
 	DrawHelper.NumCells = DrawHelper.PerspectiveGridSize / (CellSize * 2);
 
-	SetViewMode(VMI_Lit);
+	SetViewMode(VMI_Unlit);
 
 	EngineShowFlags.SetSeparateTranslucency(true);
 	EngineShowFlags.SetSnap(0);
@@ -54,8 +55,6 @@ FLidarPointCloudEditorViewportClient::FLidarPointCloudEditorViewportClient(TWeak
 
 	AdvancedPreviewScene = static_cast<FAdvancedPreviewScene*>(PreviewScene);
 
-	SetPreviewCloud(InPreviewPointCLoud, InPreviewPointCloudComponent, true);
-
 	// Register delegate to update the show flags when the post processing is turned on or off
 	UAssetViewerSettings::Get()->OnAssetViewerSettingsChanged().AddRaw(this, &FLidarPointCloudEditorViewportClient::OnAssetViewerSettingsChanged);
 
@@ -332,139 +331,147 @@ void FLidarPointCloudEditorViewportClient::OnSelectionStart()
 
 void FLidarPointCloudEditorViewportClient::OnSelectionEnd(const FIntVector4& SelectionArea, ELidarPointCloudSelectionMode Mode)
 {
-	if (PointCloudComponent)
+	if (ULidarPointCloudComponent* PointCloudComponentRawPtr = PointCloudComponent.Get())
 	{
-		// Compute a view.
-		FSceneViewFamilyContext ViewFamily(FSceneViewFamily::ConstructionValues(Viewport, GetScene(), EngineShowFlags).SetRealtimeUpdate(IsRealtime()));
-		FSceneView* View = CalcSceneView(&ViewFamily);
-
-		/**
-		 * Compute 4 rays from the corners of the selection area
-		 *
-		 * 0 -- 3
-		 * |    |
-		 * 1 -- 2
-		 */
-		FVector Origins[6];
-		FVector Normals[6];
-		FVector Directions[4];
-		FSceneView::DeprojectScreenToWorld(FVector2D(SelectionArea.X, SelectionArea.Y), FIntRect(FIntPoint(0, 0), Viewport->GetSizeXY()), View->ViewMatrices.GetInvViewProjectionMatrix(), Origins[0], Directions[0]);
-		FSceneView::DeprojectScreenToWorld(FVector2D(SelectionArea.X, SelectionArea.W), FIntRect(FIntPoint(0, 0), Viewport->GetSizeXY()), View->ViewMatrices.GetInvViewProjectionMatrix(), Origins[1], Directions[1]);
-		FSceneView::DeprojectScreenToWorld(FVector2D(SelectionArea.Z, SelectionArea.W), FIntRect(FIntPoint(0, 0), Viewport->GetSizeXY()), View->ViewMatrices.GetInvViewProjectionMatrix(), Origins[2], Directions[2]);
-		FSceneView::DeprojectScreenToWorld(FVector2D(SelectionArea.Z, SelectionArea.Y), FIntRect(FIntPoint(0, 0), Viewport->GetSizeXY()), View->ViewMatrices.GetInvViewProjectionMatrix(), Origins[3], Directions[3]);
-
-		const FVector& ViewDirection = View->GetViewDirection();
-
-		// Shared calculations
-		Normals[4] = ViewDirection;
-		Normals[5] = -ViewDirection;
-		Origins[4] = Origins[0] + ViewDirection * 99999999;
-
-		// Perspective View
-		if (View->IsPerspectiveProjection())
-		{			
-			Origins[5] = Origins[0];
-
-			// Calculate plane normals
-			Normals[0] = Directions[1] ^ Directions[0];
-			Normals[1] = Directions[2] ^ Directions[1];
-			Normals[2] = Directions[3] ^ Directions[2];
-			Normals[3] = Directions[0] ^ Directions[3];
-
-			Normals[0].Normalize();
-			Normals[1].Normalize();
-			Normals[2].Normalize();
-			Normals[3].Normalize();
-		}
-		// Ortho Views
-		else
+		if (auto Editor = PointCloudEditorPtr.Pin())
 		{
-			Origins[5] = -Origins[4];
+			// Compute a view.
+			FSceneViewFamilyContext ViewFamily(FSceneViewFamily::ConstructionValues(Viewport, GetScene(), EngineShowFlags).SetRealtimeUpdate(IsRealtime()));
+			FSceneView* View = CalcSceneView(&ViewFamily);	
 
-			// Left / Right view
-			if (ViewDirection.X != 0)
+			const FVector LocationOffset = Editor->GetPointCloudBeingEdited()->LocationOffset.ToVector();
+
+			// Adjust for the LocationOffset
+			if (View->IsPerspectiveProjection())
 			{
-				Normals[0] = -ViewDirection.X * FVector::RightVector;
-				Normals[1] = -ViewDirection.X * FVector::UpVector;
-				Normals[2] = ViewDirection.X * FVector::RightVector;
-				Normals[3] = ViewDirection.X * FVector::UpVector;
+				View->ViewLocation -= LocationOffset;
+				View->UpdateViewMatrix();
 			}
-			// Front / Back view
-			else if (ViewDirection.Y != 0)
+			
+			const FMatrix InvViewProjectionMatrix = View->ViewMatrices.GetInvViewProjectionMatrix();
+
+			/**
+			 * Compute 4 rays from the corners of the selection area
+			 *
+			 * 0 -- 3
+			 * |    |
+			 * 1 -- 2
+			 */
+			FVector Origins[6];
+			FVector Normals[6];
+			FVector Directions[4];
+			FSceneView::DeprojectScreenToWorld(FVector2D(SelectionArea.X, SelectionArea.Y), FIntRect(FIntPoint(0, 0), Viewport->GetSizeXY()), InvViewProjectionMatrix, Origins[0], Directions[0]);
+			FSceneView::DeprojectScreenToWorld(FVector2D(SelectionArea.X, SelectionArea.W), FIntRect(FIntPoint(0, 0), Viewport->GetSizeXY()), InvViewProjectionMatrix, Origins[1], Directions[1]);
+			FSceneView::DeprojectScreenToWorld(FVector2D(SelectionArea.Z, SelectionArea.W), FIntRect(FIntPoint(0, 0), Viewport->GetSizeXY()), InvViewProjectionMatrix, Origins[2], Directions[2]);
+			FSceneView::DeprojectScreenToWorld(FVector2D(SelectionArea.Z, SelectionArea.Y), FIntRect(FIntPoint(0, 0), Viewport->GetSizeXY()), InvViewProjectionMatrix, Origins[3], Directions[3]);
+
+			const FVector& ViewDirection = View->GetViewDirection();
+
+			// Shared calculations
+			Normals[4] = ViewDirection;
+			Normals[5] = -ViewDirection;
+			Origins[4] = Origins[0] + ViewDirection * 99999999;
+
+			// Perspective View
+			if (View->IsPerspectiveProjection())
 			{
-				Normals[0] = ViewDirection.Y * FVector::ForwardVector;
-				Normals[1] = ViewDirection.Y * FVector::UpVector;
-				Normals[2] = -ViewDirection.Y * FVector::ForwardVector;
-				Normals[3] = -ViewDirection.Y * FVector::UpVector;
+				Origins[5] = Origins[0];
+
+				// Calculate plane normals
+				Normals[0] = Directions[1] ^ Directions[0];
+				Normals[1] = Directions[2] ^ Directions[1];
+				Normals[2] = Directions[3] ^ Directions[2];
+				Normals[3] = Directions[0] ^ Directions[3];
+
+				Normals[0].Normalize();
+				Normals[1].Normalize();
+				Normals[2].Normalize();
+				Normals[3].Normalize();
 			}
-			// Top / Bottom view
+			// Ortho Views
 			else
 			{
-				Normals[0] = ViewDirection.Z * FVector::ForwardVector;
-				Normals[1] = -ViewDirection.Z * FVector::RightVector;
-				Normals[2] = -ViewDirection.Z * FVector::ForwardVector;
-				Normals[3] = ViewDirection.Z * FVector::RightVector;
+				// Adjust for the LocationOffset
+				for (int32 i = 0; i < 5; ++i)
+				{
+					Origins[i] -= LocationOffset;
+				}
+
+				Origins[5] = -Origins[4];
+
+				// Left / Right view
+				if (ViewDirection.X != 0)
+				{
+					Normals[0] = -ViewDirection.X * FVector::RightVector;
+					Normals[1] = -ViewDirection.X * FVector::UpVector;
+					Normals[2] = ViewDirection.X * FVector::RightVector;
+					Normals[3] = ViewDirection.X * FVector::UpVector;
+				}
+				// Front / Back view
+				else if (ViewDirection.Y != 0)
+				{
+					Normals[0] = ViewDirection.Y * FVector::ForwardVector;
+					Normals[1] = ViewDirection.Y * FVector::UpVector;
+					Normals[2] = -ViewDirection.Y * FVector::ForwardVector;
+					Normals[3] = -ViewDirection.Y * FVector::UpVector;
+				}
+				// Top / Bottom view
+				else
+				{
+					Normals[0] = ViewDirection.Z * FVector::ForwardVector;
+					Normals[1] = -ViewDirection.Z * FVector::RightVector;
+					Normals[2] = -ViewDirection.Z * FVector::ForwardVector;
+					Normals[3] = ViewDirection.Z * FVector::RightVector;
+				}
 			}
-		}
 
-		// Compute selection frustum planes from the rays
-		TArray<FPlane, TInlineAllocator<6>> Planes;
+			// Compute selection frustum planes from the rays
+			TArray<FPlane, TInlineAllocator<6>> Planes;
 
-		// Left, Bottom, Right, Top, Forward, Backward
-		for(int32 i = 0; i < 6; ++i)
-		{
-			Planes.Emplace(Origins[i], Normals[i]);
-		}
+			// Left, Bottom, Right, Top, Forward, Backward
+			for (int32 i = 0; i < 6; ++i)
+			{
+				Planes.Emplace(Origins[i], Normals[i]);
+			}
 
-		if (Mode == ELidarPointCloudSelectionMode::Subtract)
-		{
-			PointCloudEditorPtr.Pin()->DeselectPointsByFrustum(FConvexVolume(Planes));
-		}
-		else
-		{
-			PointCloudEditorPtr.Pin()->SelectPointsByFrustum(FConvexVolume(Planes), Mode == ELidarPointCloudSelectionMode::Add);
+			if (Mode == ELidarPointCloudSelectionMode::Subtract)
+			{
+				Editor->DeselectPointsByFrustum(FConvexVolume(Planes));
+			}
+			else
+			{
+				Editor->SelectPointsByFrustum(FConvexVolume(Planes), Mode == ELidarPointCloudSelectionMode::Add);
+			}
 		}
 	}
 }
 
 void FLidarPointCloudEditorViewportClient::ResetCamera()
 {
-	if (PointCloudComponent != nullptr)
+	if (ULidarPointCloudComponent* PointCloudComponentRawPtr = PointCloudComponent.Get())
 	{
-		FocusViewportOnBox(PointCloudComponent->Bounds.GetBox());
+		FocusViewportOnBox(PointCloudComponentRawPtr->Bounds.GetBox());
 		Invalidate();
 	}
 }
 
 void FLidarPointCloudEditorViewportClient::ToggleShowNodes()
 {
-	if (PointCloudComponent != nullptr)
+	if (ULidarPointCloudComponent* PointCloudComponentRawPtr = PointCloudComponent.Get())
 	{
-		PointCloudComponent->bDrawNodeBounds = !PointCloudComponent->bDrawNodeBounds;
+		PointCloudComponentRawPtr->bDrawNodeBounds = !PointCloudComponentRawPtr->bDrawNodeBounds;
 		Invalidate();
 	}
 }
 
 bool FLidarPointCloudEditorViewportClient::IsSetShowNodesChecked() const
 {
-	return PointCloudComponent && PointCloudComponent->bDrawNodeBounds;
-}
-
-void FLidarPointCloudEditorViewportClient::SetPreviewCloud(ULidarPointCloud* InPointCloud, ULidarPointCloudComponent* InPointCloudComponent, bool bResetCamera)
-{
-	PointCloud = InPointCloud;
-	PointCloudComponent = InPointCloudComponent;
-
-	if (PointCloudComponent != nullptr)
-	{
-		PointCloudComponent->bOwnedByEditor = true;
-		PointCloudComponent->MarkRenderStateDirty();
-	}
-
-	if (bResetCamera)
+	if (ULidarPointCloudComponent* PointCloudComponentRawPtr = PointCloudComponent.Get())
 	{
-		ResetCamera();
+		return PointCloudComponentRawPtr->bDrawNodeBounds;
 	}
+	
+	return false;
 }
 
 #undef LOCTEXT_NAMESPACE 
\ No newline at end of file
diff --git a/Engine/Plugins/Enterprise/LidarPointCloud/Source/LidarPointCloudEditor/Private/LidarPointCloudEditorViewportClient.h b/Engine/Plugins/Enterprise/LidarPointCloud/Source/LidarPointCloudEditor/Private/LidarPointCloudEditorViewportClient.h
--- a/Engine/Plugins/Enterprise/LidarPointCloud/Source/LidarPointCloudEditor/Private/LidarPointCloudEditorViewportClient.h
+++ b/Engine/Plugins/Enterprise/LidarPointCloud/Source/LidarPointCloudEditor/Private/LidarPointCloudEditorViewportClient.h
@@ -27,7 +27,7 @@ enum class ELidarPointCloudSelectionMode : uint8
 class FLidarPointCloudEditorViewportClient : public FEditorViewportClient, public TSharedFromThis<FLidarPointCloudEditorViewportClient>
 {
 public:
-	FLidarPointCloudEditorViewportClient(TWeakPtr<FLidarPointCloudEditor> InPointCloudEditor, const TSharedRef<SLidarPointCloudEditorViewport>& InPointCloudEditorViewport, const TSharedRef<FAdvancedPreviewScene>& InPreviewScene, ULidarPointCloud* InPreviewPointCLoud, ULidarPointCloudComponent* InPreviewPointCloudComponent);
+	FLidarPointCloudEditorViewportClient(TWeakPtr<FLidarPointCloudEditor> InPointCloudEditor, const TSharedRef<SLidarPointCloudEditorViewport>& InPointCloudEditorViewport, FAdvancedPreviewScene* InPreviewScene, ULidarPointCloud* InPreviewPointCloud, ULidarPointCloudComponent* InPreviewPointCloudComponent);
 	~FLidarPointCloudEditorViewportClient();
 
 	// FEditorViewportClient interface
@@ -50,14 +50,6 @@ public:
 	virtual ECoordSystem GetWidgetCoordSystemSpace() const override { return COORD_Local; }
 	virtual bool ShouldOrbitCamera() const override;
 
-	/**
-	 *	Updates the point cloud and point cloud component being used in the Point Cloud Editor.
-	 *
-	 *	@param	InPointCloud				The point cloud handled by the editor.
-	 *	@param	InPointCloudComponent		The point cloud component from the viewport's scene.
-	 */
-	void SetPreviewCloud(ULidarPointCloud* InPointCloud, ULidarPointCloudComponent* InPointCloudComponent, bool bResetCamera = true);
-
 	void ResetCamera();
 
 	/** Callback for toggling the nodes show flag. */
@@ -81,10 +73,7 @@ private:
 
 private:
 	/** Component for the point cloud. */
-	ULidarPointCloudComponent* PointCloudComponent;
-
-	/** The point cloud being used in the editor. */
-	ULidarPointCloud* PointCloud;
+	TWeakObjectPtr<ULidarPointCloudComponent> PointCloudComponent;
 
 	/** Pointer back to the PointCloud editor tool that owns us */
 	TWeakPtr<FLidarPointCloudEditor> PointCloudEditorPtr;
diff --git a/Engine/Plugins/Enterprise/LidarPointCloud/Source/LidarPointCloudRuntime/Private/Collision/LidarPointCloudCollision.cpp b/Engine/Plugins/Enterprise/LidarPointCloud/Source/LidarPointCloudRuntime/Private/Collision/LidarPointCloudCollision.cpp
--- a/Engine/Plugins/Enterprise/LidarPointCloud/Source/LidarPointCloudRuntime/Private/Collision/LidarPointCloudCollision.cpp
+++ b/Engine/Plugins/Enterprise/LidarPointCloud/Source/LidarPointCloudRuntime/Private/Collision/LidarPointCloudCollision.cpp
@@ -3,6 +3,7 @@
 #include "LidarPointCloudCollision.h"
 #include "LidarPointCloudShared.h"
 #include "LidarPointCloudOctree.h"
+#include "LidarPointCloud.h"
 
 #include "Interfaces/Interface_CollisionDataProvider.h"
 #include "Async/Async.h"
@@ -385,7 +386,7 @@ void MarchingCubes(uint8* VoxelizedGrid, int32 GridDimension, float CellSize, FV
 	}
 }
 
-void LidarPointCloudCollision::BuildCollisionMesh(const FLidarPointCloudOctree* Octree, const float& CellSize, const bool& bVisibleOnly, FTriMeshCollisionData* CollisionMesh)
+void LidarPointCloudCollision::BuildCollisionMesh(FLidarPointCloudOctree* Octree, const float& CellSize, const bool& bVisibleOnly, FTriMeshCollisionData* CollisionMesh)
 {
 	if (!CollisionMesh)
 	{
@@ -397,6 +398,7 @@ void LidarPointCloudCollision::BuildCollisionMesh(const FLidarPointCloudOctree*
 	// Expand original bounds to make sure the mesh is closed at the edges
 	const FBox Bounds = Octree->GetBounds().ExpandBy(FVector::OneVector * CellSize, FVector::OneVector * CellSize);
 	const FVector BoundsSize = Bounds.GetSize();
+	const FVector LocationOffset = Octree->GetOwner()->LocationOffset.ToVector();
 
 	// Number of cells in each axis
 	const int32 BatchSize = GetDefault<ULidarPointCloudSettings>()->MeshingBatchSize;
@@ -424,7 +426,7 @@ void LidarPointCloudCollision::BuildCollisionMesh(const FLidarPointCloudOctree*
 	// Fire threads
 	for (int32 t = 0; t < MaxNumThreads && SampleIndex.GetValue() < TotalNumSamples; t++)
 	{
-		ThreadResults.Add(Async(EAsyncExecution::TaskGraph, [t, &SampleIndex, &Vertices, &NumVertices, CollisionMesh, BatchSize, TotalNumSamples, NumSamples, BaseSamplingBounds, CellSize, Octree, InversedCellSize, bVisibleOnly]
+		ThreadResults.Add(Async(EAsyncExecution::TaskGraph, [t, &SampleIndex, &Vertices, &NumVertices, CollisionMesh, BatchSize, TotalNumSamples, NumSamples, BaseSamplingBounds, CellSize, Octree, InversedCellSize, bVisibleOnly, LocationOffset]
 		{
 			// Local caching arrays to reduce number of syncs required
 			TArray<FVector> _Vertices;
@@ -467,7 +469,7 @@ void LidarPointCloudCollision::BuildCollisionMesh(const FLidarPointCloudOctree*
 					VoxelizedGrid[Grid.Z * NumCellsSq + Grid.Y * BatchSize + Grid.X] = 1;
 				}
 
-				MarchingCubes(VoxelizedGrid, BatchSize, CellSize, SamplingBounds.Min, _Vertices);
+				MarchingCubes(VoxelizedGrid, BatchSize, CellSize, SamplingBounds.Min + LocationOffset, _Vertices);
 			}
 
 			// Sync
diff --git a/Engine/Plugins/Enterprise/LidarPointCloud/Source/LidarPointCloudRuntime/Private/Collision/LidarPointCloudCollision.h b/Engine/Plugins/Enterprise/LidarPointCloud/Source/LidarPointCloudRuntime/Private/Collision/LidarPointCloudCollision.h
--- a/Engine/Plugins/Enterprise/LidarPointCloud/Source/LidarPointCloudRuntime/Private/Collision/LidarPointCloudCollision.h
+++ b/Engine/Plugins/Enterprise/LidarPointCloud/Source/LidarPointCloudRuntime/Private/Collision/LidarPointCloudCollision.h
@@ -9,5 +9,5 @@ struct FTriMeshCollisionData;
 
 namespace LidarPointCloudCollision
 {
-	void BuildCollisionMesh(const FLidarPointCloudOctree* Octree, const float& CellSize, const bool& bVisibleOnly, FTriMeshCollisionData* CollisionMesh);
+	void BuildCollisionMesh(FLidarPointCloudOctree* Octree, const float& CellSize, const bool& bVisibleOnly, FTriMeshCollisionData* CollisionMesh);
 }
\ No newline at end of file
diff --git a/Engine/Plugins/Enterprise/LidarPointCloud/Source/LidarPointCloudRuntime/Private/IO/LidarPointCloudFileIO_ASCII.cpp b/Engine/Plugins/Enterprise/LidarPointCloud/Source/LidarPointCloudRuntime/Private/IO/LidarPointCloudFileIO_ASCII.cpp
--- a/Engine/Plugins/Enterprise/LidarPointCloud/Source/LidarPointCloudRuntime/Private/IO/LidarPointCloudFileIO_ASCII.cpp
+++ b/Engine/Plugins/Enterprise/LidarPointCloud/Source/LidarPointCloudRuntime/Private/IO/LidarPointCloudFileIO_ASCII.cpp
@@ -48,8 +48,6 @@ bool ULidarPointCloudFileIO_ASCII::HandleImport(const FString& Filename, TShared
 		TArray<int32> SelectedColumns = Settings->SelectedColumns;
 		FVector2D RGBRange = Settings->RGBRange;
 
-		const bool bAutoCenter = GetDefault<ULidarPointCloudSettings>()->bAutoCenterOnImport;
-
 		// Flags which columns should have data assigned and used
 		TArray<bool> IsColumnPopulated;
 		IsColumnPopulated.AddZeroed(SelectedColumns.Num());
@@ -130,7 +128,7 @@ bool ULidarPointCloudFileIO_ASCII::HandleImport(const FString& Filename, TShared
 			const int64 MaxBufferSize = GetDefault<ULidarPointCloudSettings>()->MaxImportBufferSize;
 			const float ImportScale = GetDefault<ULidarPointCloudSettings>()->ImportScale;
 
-			bool bFirstPointSet = !bAutoCenter;
+			bool bFirstPointSet = false;
 
 			// Prepare pointers
 			uint8 *Data = (uint8*)FMemory::Malloc(MaxBufferSize);
@@ -236,12 +234,6 @@ bool ULidarPointCloudFileIO_ASCII::HandleImport(const FString& Filename, TShared
 			FMemory::Free(Data);
 		}
 
-		// Apply re-centering, if selected
-		if (bAutoCenter)
-		{
-			OutImportResults.CenterPoints();
-		}
-
 		Reader->Close();
 
 		return !OutImportResults.IsCancelled();
@@ -270,7 +262,7 @@ bool ULidarPointCloudFileIO_ASCII::HandleExport(const FString& Filename, class U
 
 		for (FLidarPointCloudPoint* Point : Points)
 		{
-			FDoubleVector Location = (PointCloud->OriginalCoordinates + Point->Location) * ExportScale;
+			FDoubleVector Location = (PointCloud->LocationOffset + Point->Location) * ExportScale;
 			Lines.Emplace(FString::Printf(TEXT("%f,%f,%f,%d,%d,%d,%d"), Location.X, -Location.Y, Location.Z, Point->Color.A, Point->Color.R, Point->Color.G, Point->Color.B));
 		}
 
diff --git a/Engine/Plugins/Enterprise/LidarPointCloud/Source/LidarPointCloudRuntime/Private/IO/LidarPointCloudFileIO_LAS.cpp b/Engine/Plugins/Enterprise/LidarPointCloud/Source/LidarPointCloudRuntime/Private/IO/LidarPointCloudFileIO_LAS.cpp
--- a/Engine/Plugins/Enterprise/LidarPointCloud/Source/LidarPointCloudRuntime/Private/IO/LidarPointCloudFileIO_LAS.cpp
+++ b/Engine/Plugins/Enterprise/LidarPointCloud/Source/LidarPointCloudRuntime/Private/IO/LidarPointCloudFileIO_LAS.cpp
@@ -84,115 +84,6 @@ void FLidarPointCloudImportSettings_LAS::ReadFileHeader(const FString& InFilenam
 	}
 }
 
-/** Used to help track multiple import buffer allocations */
-class FLidarPointCloudFileIO_LAS_ImportBuffer
-{
-public:
-	FLidarPointCloudFileIO_LAS_ImportBuffer()
-		: bInUse(false)
-	{
-	}
-	~FLidarPointCloudFileIO_LAS_ImportBuffer() = default;
-	FLidarPointCloudFileIO_LAS_ImportBuffer(const FLidarPointCloudFileIO_LAS_ImportBuffer& Other)
-	{
-		Data = Other.Data;
-		bInUse = false;
-	}
-	FLidarPointCloudFileIO_LAS_ImportBuffer(FLidarPointCloudFileIO_LAS_ImportBuffer&&) = delete;
-	FLidarPointCloudFileIO_LAS_ImportBuffer& operator=(const FLidarPointCloudFileIO_LAS_ImportBuffer& Other)
-	{
-		Data = Other.Data;
-		bInUse = false;
-		return *this;
-	}
-	FLidarPointCloudFileIO_LAS_ImportBuffer& operator=(FLidarPointCloudFileIO_LAS_ImportBuffer&&) = delete;
-
-	FORCEINLINE uint8* GetData() { return Data.GetData(); }
-
-	void Initialize(const int32& Size)
-	{
-		Data.AddUninitialized(Size);
-	}
-
-public:
-	TAtomic<bool> bInUse;
-
-private:
-	TArray<uint8> Data;
-};
-
-/** Used to help track multiple import buffer allocations */
-class FLidarPointCloudFileIO_LAS_ImportBufferManager
-{
-public:
-	FLidarPointCloudFileIO_LAS_ImportBufferManager(const int32& BufferSize)
-		: BufferSize(BufferSize)
-		, Head(FLidarPointCloudFileIO_LAS_ImportBuffer())
-		, Tail(&Head)
-	{
-		Head.Element.Initialize(BufferSize);
-	}
-	~FLidarPointCloudFileIO_LAS_ImportBufferManager()
-	{
-		auto Iterator = &Head;
-		while (Iterator)
-		{
-			if (Iterator != &Head)
-			{
-				auto Tmp = Iterator;
-				Iterator = Iterator->Next;
-				delete Tmp;
-			}
-			else
-			{
-				Iterator = Iterator->Next;
-			}
-		}
-	}
-	FLidarPointCloudFileIO_LAS_ImportBufferManager(const FLidarPointCloudFileIO_LAS_ImportBufferManager&) = delete;
-	FLidarPointCloudFileIO_LAS_ImportBufferManager(FLidarPointCloudFileIO_LAS_ImportBufferManager&&) = delete;
-	FLidarPointCloudFileIO_LAS_ImportBufferManager& operator=(const FLidarPointCloudFileIO_LAS_ImportBufferManager&) = delete;
-	FLidarPointCloudFileIO_LAS_ImportBufferManager& operator=(FLidarPointCloudFileIO_LAS_ImportBufferManager&&) = delete;
-
-	FLidarPointCloudFileIO_LAS_ImportBuffer* GetFreeBuffer()
-	{
-		FLidarPointCloudFileIO_LAS_ImportBuffer* OutBuffer = nullptr;
-
-		// Find available memory allocation
-		{
-			auto Iterator = &Head;
-			while (Iterator)
-			{
-				if (!Iterator->Element.bInUse)
-				{
-					OutBuffer = &Iterator->Element;
-					break;
-				}
-
-				Iterator = Iterator->Next;
-			}
-		}
-
-		// If none found, add a new one
-		if (!OutBuffer)
-		{
-			Tail->Next = new TList<FLidarPointCloudFileIO_LAS_ImportBuffer>(FLidarPointCloudFileIO_LAS_ImportBuffer());
-			Tail = Tail->Next;
-			OutBuffer = &Tail->Element;
-			OutBuffer->Initialize(BufferSize);
-		}
-
-		OutBuffer->bInUse = true;
-
-		return OutBuffer;
-	}
-
-private:
-	int32 BufferSize;
-	TList<FLidarPointCloudFileIO_LAS_ImportBuffer> Head;
-	TList<FLidarPointCloudFileIO_LAS_ImportBuffer>* Tail;
-};
-
 bool ULidarPointCloudFileIO_LAS::HandleImport(const FString& Filename, TSharedPtr<FLidarPointCloudImportSettings> ImportSettings, FLidarPointCloudImportResults &OutImportResults)
 {
 	if (!ValidateImportSettings(ImportSettings, Filename))
@@ -289,21 +180,20 @@ bool ULidarPointCloudFileIO_LAS::HandleImport(const FString& Filename, TSharedPt
 				const float IntensityMultiplier = 1.0f / (bUse16BitIntensity ? 65535.0f : 255.0f);
 				const float RGBMultiplier = 1.0f / (bUse16BitRGB ? 65535.0f : 255.0f);
 
-				const bool bAutoCenter = GetDefault<ULidarPointCloudSettings>()->bAutoCenterOnImport;
 				const float ImportScale = GetDefault<ULidarPointCloudSettings>()->ImportScale;
 
-				FThreadSafeBool bFirstPointSet = !bAutoCenter;
+				FThreadSafeBool bFirstPointSet = false;
 
 				// Multi-threading
 				FCriticalSection CoordsLock;
 				FCriticalSection PointsLock;
 				TArray<TFuture<void>> ThreadResults;
-				FLidarPointCloudFileIO_LAS_ImportBufferManager BufferManager(MaxPointsToRead * Header.PointDataRecordLength);
+				FLidarPointCloudDataBufferManager BufferManager(MaxPointsToRead * Header.PointDataRecordLength);
 
 				// Stream the data
 				while (PointsRead < TotalPointsToRead && !OutImportResults.IsCancelled())
 				{
-					FLidarPointCloudFileIO_LAS_ImportBuffer* Buffer = BufferManager.GetFreeBuffer();
+					FLidarPointCloudDataBuffer* Buffer = BufferManager.GetFreeBuffer();
 
 					// Data should never be null
 					check(Buffer->GetData());
@@ -406,8 +296,7 @@ bool ULidarPointCloudFileIO_LAS::HandleImport(const FString& Filename, TSharedPt
 								}
 							}
 
-							// Memory Sync
-							Buffer->bInUse = false;
+							Buffer->MarkAsFree();
 						}));
 
 					PointsRead += PointsToRead;
@@ -416,12 +305,6 @@ bool ULidarPointCloudFileIO_LAS::HandleImport(const FString& Filename, TSharedPt
 				// Sync threads
 				for (auto& ThreadResult : ThreadResults) ThreadResult.Get();
 
-				// Apply re-centering, if selected and not applied already
-				if (bAutoCenter)
-				{
-					OutImportResults.CenterPoints();
-				}
-
 				// Make sure to progress the counter to the end before returning
 				OutImportResults.IncrementProgressCounter(TotalPointsToRead);
 			}
@@ -447,10 +330,6 @@ bool ULidarPointCloudFileIO_LAS::HandleExport(const FString& Filename, ULidarPoi
 	FDoubleVector Min = PointCloud->GetBounds().Min;
 	FDoubleVector Max = PointCloud->GetBounds().Max;
 
-	// Add original coordinates
-	Min += PointCloud->OriginalCoordinates;
-	Max += PointCloud->OriginalCoordinates;
-
 	// Flip Y
 	float MaxY = Max.Y;
 	Max.Y = -Min.Y;
@@ -484,7 +363,7 @@ bool ULidarPointCloudFileIO_LAS::HandleExport(const FString& Filename, ULidarPoi
 
 		for (FLidarPointCloudPoint* Point : Points)
 		{
-			FDoubleVector Location = (PointCloud->OriginalCoordinates + Point->Location) * ExportScale;
+			FDoubleVector Location = (PointCloud->LocationOffset + Point->Location) * ExportScale;
 			Location.Y = -Location.Y;
 
 			PointRecord->Location = (ForwardScale * (Location - Min)).ToIntVector();
diff --git a/Engine/Plugins/Enterprise/LidarPointCloud/Source/LidarPointCloudRuntime/Private/LidarPointCloud.cpp b/Engine/Plugins/Enterprise/LidarPointCloud/Source/LidarPointCloudRuntime/Private/LidarPointCloud.cpp
--- a/Engine/Plugins/Enterprise/LidarPointCloud/Source/LidarPointCloudRuntime/Private/LidarPointCloud.cpp
+++ b/Engine/Plugins/Enterprise/LidarPointCloud/Source/LidarPointCloudRuntime/Private/LidarPointCloud.cpp
@@ -30,11 +30,40 @@
 #define IS_PROPERTY(Name) PropertyChangedEvent.MemberProperty->GetName().Equals(#Name)
 
 const FGuid ULidarPointCloud::PointCloudFileGUID('P', 'C', 'P', 'F');
-const int32 ULidarPointCloud::PointCloudFileVersion(15);
+const int32 ULidarPointCloud::PointCloudFileVersion(17);
 FCustomVersionRegistration PCPFileVersion(ULidarPointCloud::PointCloudFileGUID, ULidarPointCloud::PointCloudFileVersion, TEXT("LiDAR Point Cloud File Version"));
 
 #define LOCTEXT_NAMESPACE "LidarPointCloud"
 
+class FPointCloudLatentAction : public FPendingLatentAction
+{
+public:
+	FName ExecutionFunction;
+	int32 OutputLink;
+	FWeakObjectPtr CallbackTarget;
+	ELidarPointCloudAsyncMode* Mode = nullptr;
+
+	FPointCloudLatentAction(const FLatentActionInfo& LatentInfo, ELidarPointCloudAsyncMode& Mode)
+		: ExecutionFunction(LatentInfo.ExecutionFunction)
+		, OutputLink(LatentInfo.Linkage)
+		, CallbackTarget(LatentInfo.CallbackTarget)
+		, Mode(&Mode)
+	{
+	}
+
+	virtual void UpdateOperation(FLatentResponse& Response) override
+	{
+		if (*Mode != ELidarPointCloudAsyncMode::Progress)
+		{
+			Response.FinishAndTriggerIf(true, ExecutionFunction, OutputLink, CallbackTarget);
+		}
+		else
+		{
+			Response.TriggerLink(ExecutionFunction, OutputLink, CallbackTarget);
+		}
+	}
+};
+
 /////////////////////////////////////////////////
 // FLidarPointCloudNotification
 
@@ -160,8 +189,9 @@ void FLidarPointCloudNotification::UpdateStatus()
 
 ULidarPointCloud::ULidarPointCloud()
 	: CollisionAccuracy(100)
-	, Octree()
+	, Octree(this)
 	, OriginalCoordinates(FDoubleVector::ZeroVector)
+	, LocationOffset(FDoubleVector::ZeroVector)
 	, Notification(this)
 	, BodySetup(nullptr)
 	, bCollisionBuildInProgress(false)
@@ -198,8 +228,6 @@ void ULidarPointCloud::Serialize(FArchive& Ar)
 	{
 		return;
 	}
-
-	FScopeBenchmarkTimer BenchmarkTimer("Serialization");
 	
 	ULidarPointCloudFileIO::SerializeImportSettings(Ar, ImportSettings);
 
@@ -250,6 +278,12 @@ void ULidarPointCloud::BeginDestroy()
 	ReleaseCollisionRendering();
 }
 
+void ULidarPointCloud::PreSave(const class ITargetPlatform* TargetPlatform)
+{
+	Super::PreSave(TargetPlatform);
+	OnPreSaveCleanupEvent.Broadcast();
+}
+
 #if WITH_EDITOR
 void ULidarPointCloud::PostEditChangeProperty(FPropertyChangedEvent& PropertyChangedEvent)
 {
@@ -293,24 +327,24 @@ TArray<FLidarPointCloudPoint*> ULidarPointCloud::GetPoints(int64 StartIndex, int
 	return Points;
 }
 
-TArray<FLidarPointCloudPoint> ULidarPointCloud::GetPointsAsCopies(int32 StartIndex, int32 Count) const
+TArray<FLidarPointCloudPoint> ULidarPointCloud::GetPointsAsCopies(bool bReturnWorldSpace, int32 StartIndex, int32 Count) const
 {
 	TArray<FLidarPointCloudPoint> Points;
-	GetPointsAsCopies(Points, StartIndex, Count);
+	GetPointsAsCopies(Points, bReturnWorldSpace, StartIndex, Count);
 	return Points;
 }
 
-TArray<FLidarPointCloudPoint> ULidarPointCloud::GetPointsInSphereAsCopies(FVector Center, float Radius, bool bVisibleOnly)
+TArray<FLidarPointCloudPoint> ULidarPointCloud::GetPointsInSphereAsCopies(FVector Center, float Radius, bool bVisibleOnly, bool bReturnWorldSpace)
 {
 	TArray<FLidarPointCloudPoint> Points;
-	GetPointsInSphereAsCopies(Points, FSphere(Center, Radius), bVisibleOnly);
+	GetPointsInSphereAsCopies(Points, FSphere(Center, Radius), bVisibleOnly, bReturnWorldSpace);
 	return Points;
 }
 
-TArray<FLidarPointCloudPoint> ULidarPointCloud::GetPointsInBoxAsCopies(FVector Center, FVector Extent, bool bVisibleOnly)
+TArray<FLidarPointCloudPoint> ULidarPointCloud::GetPointsInBoxAsCopies(FVector Center, FVector Extent, bool bVisibleOnly, bool bReturnWorldSpace)
 {
 	TArray<FLidarPointCloudPoint> Points;
-	GetPointsInBoxAsCopies(Points, FBox(Center - Extent, Center + Extent), bVisibleOnly);
+	GetPointsInBoxAsCopies(Points, FBox(Center - Extent, Center + Extent), bVisibleOnly, bReturnWorldSpace);
 	return Points;
 }
 
@@ -402,37 +436,13 @@ void ULidarPointCloud::RemoveCollision()
 	bCollisionBuildInProgress = false;
 }
 
-void ULidarPointCloud::CenterPoints()
-{
-	// Only shift, if not already centered
-	if (!IsCentered())
-	{
-		ShiftPointsBy(-GetPointsBounds().GetCenter(), true);
-	}
-}
-
-void ULidarPointCloud::RestoreOriginalCoordinates()
+void ULidarPointCloud::SetLocationOffset(FDoubleVector Offset)
 {
-	if (!OriginalCoordinates.IsNearlyZero(0.001f))
-	{
-		ShiftPointsBy(OriginalCoordinates, true);
-	}
-}
-
-void ULidarPointCloud::ShiftPointsBy(FDoubleVector Offset, bool bRefreshPointsBounds)
-{
-	Octree.ShiftPointsBy(Offset, bRefreshPointsBounds);
-	OriginalCoordinates -= Offset;
+	LocationOffset = Offset;
 	MarkPackageDirty();
 	OnPointCloudRebuiltEvent.Broadcast();
 }
 
-bool ULidarPointCloud::IsCentered() const
-{
-	// Return as centered if within 1mm of center
-	return GetPointsBounds().GetCenter().IsNearlyZero(0.1f);
-}
-
 void ULidarPointCloud::Reimport(const FLidarPointCloudAsyncParameters& AsyncParameters)
 {
 	if (FPaths::FileExists(SourcePath.FilePath))
@@ -448,8 +458,10 @@ void ULidarPointCloud::Reimport(const FLidarPointCloudAsyncParameters& AsyncPara
 		bAsyncCancelled = false;
 		Notification.Create("Importing Point Cloud", &bAsyncCancelled);
 
+		const bool bCenter = GetDefault<ULidarPointCloudSettings>()->bAutoCenterOnImport;
+
 		// The actual import function to be executed
-		auto ImportFunction = [this, AsyncParameters]
+		auto ImportFunction = [this, AsyncParameters, bCenter]
 		{
 			// This will take over the lock
 			FScopeLock Lock(&ImportLock);
@@ -473,11 +485,11 @@ void ULidarPointCloud::Reimport(const FLidarPointCloudAsyncParameters& AsyncPara
 				FScopeLock DataLock(&Octree.DataLock);
 				
 				// Re-initialize the Octree
-				Octree.Initialize(ImportResults.Bounds);
+				Initialize(ImportResults.Bounds);
 
 				FScopeBenchmarkTimer BenchmarkTimer("Octree Build-Up");
 
-				bSuccess = InsertPoints_NoLock(ImportResults.Points.GetData(), ImportResults.Points.Num(), GetDefault<ULidarPointCloudSettings>()->DuplicateHandling, false, &bAsyncCancelled, [this, AsyncParameters](float Progress)
+				bSuccess = InsertPoints_NoLock(ImportResults.Points.GetData(), ImportResults.Points.Num(), GetDefault<ULidarPointCloudSettings>()->DuplicateHandling, false, -LocationOffset.ToVector(), &bAsyncCancelled, [this, AsyncParameters](float Progress)
 				{
 					Notification.SetProgress(50.0f + 50.0f * Progress);
 					if (AsyncParameters.ProgressCallback)
@@ -488,8 +500,11 @@ void ULidarPointCloud::Reimport(const FLidarPointCloudAsyncParameters& AsyncPara
 
 				if (bSuccess)
 				{
-					RefreshPointsBounds();
-					OriginalCoordinates = ImportResults.OriginalCoordinates;
+					RefreshBounds();
+					OriginalCoordinates = LocationOffset + ImportResults.OriginalCoordinates;
+
+					// Show the cloud at its original location, if selected
+					LocationOffset = bCenter ? FDoubleVector::ZeroVector : OriginalCoordinates;
 				}
 				else
 				{
@@ -497,6 +512,9 @@ void ULidarPointCloud::Reimport(const FLidarPointCloudAsyncParameters& AsyncPara
 
 					Octree.Empty(true);
 
+					OriginalCoordinates = FDoubleVector::ZeroVector;
+					LocationOffset = FDoubleVector::ZeroVector;
+
 					// Update PointCloudAssetRegistryCache
 					PointCloudAssetRegistryCache.PointCount = FString::FromInt(Octree.GetNumPoints());
 				}	
@@ -553,23 +571,48 @@ void ULidarPointCloud::Reimport(const FLidarPointCloudAsyncParameters& AsyncPara
 	}
 }
 
+void ULidarPointCloud::Reimport(UObject* WorldContextObject, bool bUseAsync, FLatentActionInfo LatentInfo, ELidarPointCloudAsyncMode& AsyncMode, float& Progress)
+{
+	if (UWorld* World = GEngine->GetWorldFromContextObject(WorldContextObject, EGetWorldErrorMode::LogAndReturnNull))
+	{
+		FLatentActionManager& LatentActionManager = World->GetLatentActionManager();
+		if (LatentActionManager.FindExistingAction<FPointCloudLatentAction>(LatentInfo.CallbackTarget, LatentInfo.UUID) == nullptr)
+		{
+			AsyncMode = ELidarPointCloudAsyncMode::Progress;
+			FPointCloudLatentAction* CompletionAction = new FPointCloudLatentAction(LatentInfo, AsyncMode);
+
+			LatentActionManager.AddNewAction(LatentInfo.CallbackTarget, LatentInfo.UUID, CompletionAction);
+
+			Reimport(FLidarPointCloudAsyncParameters(bUseAsync,
+				[&Progress, &AsyncMode](float InProgress)
+				{
+					Progress = InProgress;
+				},
+				[&AsyncMode](bool bSuccess)
+				{
+					AsyncMode = bSuccess ? ELidarPointCloudAsyncMode::Success : ELidarPointCloudAsyncMode::Failure;
+				}));
+		}
+	}
+}
+
 bool ULidarPointCloud::Export(const FString& Filename)
 {
 	return ULidarPointCloudFileIO::Export(Filename, this);
 }
 
-void ULidarPointCloud::InsertPoint(const FLidarPointCloudPoint& Point, ELidarPointCloudDuplicateHandling DuplicateHandling, bool bRefreshPointsBounds)
+void ULidarPointCloud::InsertPoint(const FLidarPointCloudPoint& Point, ELidarPointCloudDuplicateHandling DuplicateHandling, bool bRefreshPointsBounds, const FVector& Translation)
 {
 	FScopeLock Lock(&Octree.DataLock);
 
-	Octree.InsertPoint(&Point, DuplicateHandling, bRefreshPointsBounds);
+	Octree.InsertPoint(&Point, DuplicateHandling, bRefreshPointsBounds, Translation);
 
 	// Update PointCloudAssetRegistryCache
 	PointCloudAssetRegistryCache.PointCount = FString::FromInt(Octree.GetNumPoints());
 }
 
 template<typename T>
-bool ULidarPointCloud::InsertPoints_NoLock(T InPoints, const int64& Count, ELidarPointCloudDuplicateHandling DuplicateHandling, bool bRefreshPointsBounds, FThreadSafeBool* bCanceled, TFunction<void(float)> ProgressCallback)
+bool ULidarPointCloud::InsertPoints_NoLock(T InPoints, const int64& Count, ELidarPointCloudDuplicateHandling DuplicateHandling, bool bRefreshPointsBounds, const FVector& Translation, FThreadSafeBool* bCanceled, TFunction<void(float)> ProgressCallback)
 {
 	const ULidarPointCloudSettings* Settings = GetDefault<ULidarPointCloudSettings>();
 	const int32 MaxBatchSize = Settings->MultithreadingInsertionBatchSize;
@@ -583,7 +626,7 @@ bool ULidarPointCloud::InsertPoints_NoLock(T InPoints, const int64& Count, ELida
 
 		const int32 NumThreads = FMath::Min(FPlatformMisc::NumberOfCoresIncludingHyperthreads() - 1, (int32)(Count / MaxBatchSize) + 1);
 		TArray<TFuture<void>>ThreadResults;
-		ThreadResults.Reserve( NumThreads );
+		ThreadResults.Reserve(NumThreads);
 		const int64 NumPointsPerThread = Count / NumThreads + 1;
 
 		FCriticalSection ProgressCallbackLock;
@@ -592,7 +635,8 @@ bool ULidarPointCloud::InsertPoints_NoLock(T InPoints, const int64& Count, ELida
 		for (int32 t = 0; t < NumThreads; t++)
 		{
 			const int32 ThreadID = t;
-			ThreadResults.Add(Async(EAsyncExecution::Thread, [this, ThreadID, DuplicateHandling, bRefreshPointsBounds, MaxBatchSize, NumPointsPerThread, RefreshStatusFrequency, &ProcessedPoints, &TotalProcessedPoints, InPoints, Count, &ProgressCallback, &ProgressCallbackLock, &bCanceled]
+
+			ThreadResults.Add(Async(EAsyncExecution::Thread, [this, ThreadID, DuplicateHandling, bRefreshPointsBounds, MaxBatchSize, NumPointsPerThread, RefreshStatusFrequency, &ProcessedPoints, &TotalProcessedPoints, InPoints, Count, &ProgressCallback, &ProgressCallbackLock, &bCanceled, &Translation]
 				{
 					int64 Idx = ThreadID * NumPointsPerThread;
 					int64 MaxIdx = FMath::Min(Idx + NumPointsPerThread, Count);
@@ -602,7 +646,7 @@ bool ULidarPointCloud::InsertPoints_NoLock(T InPoints, const int64& Count, ELida
 					{
 						int32 BatchSize = FMath::Min(MaxIdx - Idx, (int64)MaxBatchSize);
 
-						Octree.InsertPoints(DataPointer, BatchSize, DuplicateHandling, bRefreshPointsBounds);
+						Octree.InsertPoints(DataPointer, BatchSize, DuplicateHandling, bRefreshPointsBounds, Translation);
 
 						if (ProgressCallback)
 						{
@@ -642,7 +686,7 @@ bool ULidarPointCloud::InsertPoints_NoLock(T InPoints, const int64& Count, ELida
 		{
 			int32 BatchSize = FMath::Min(Count - Idx, (int64)MaxBatchSize);
 
-			Octree.InsertPoints(DataPointer, BatchSize, DuplicateHandling, bRefreshPointsBounds);
+			Octree.InsertPoints(DataPointer, BatchSize, DuplicateHandling, bRefreshPointsBounds, Translation);
 
 			Idx += BatchSize;
 			DataPointer += BatchSize;
@@ -679,13 +723,9 @@ bool ULidarPointCloud::SetData(T Points, const int64& Count, TFunction<void(floa
 		// Initialize the Octree
 		Initialize(Bounds);
 
-		bSuccess = InsertPoints_NoLock(Points, Count, GetDefault<ULidarPointCloudSettings>()->DuplicateHandling, false, nullptr, MoveTemp(ProgressCallback));
+		bSuccess = InsertPoints_NoLock(Points, Count, GetDefault<ULidarPointCloudSettings>()->DuplicateHandling, false, -LocationOffset.ToVector(), nullptr, MoveTemp(ProgressCallback));
 
-		if (bSuccess)
-		{
-			RefreshPointsBounds();
-		}
-		else
+		if (!bSuccess)
 		{
 			Octree.Empty(true);
 		}
@@ -757,7 +797,9 @@ void ULidarPointCloud::Merge(TArray<ULidarPointCloud*> PointCloudsToMerge, TFunc
 
 	// Make a copy of current points, as the data will be reinitialized
 	TArray<FLidarPointCloudPoint> Points;
-	GetPointsAsCopies(Points);
+	GetPointsAsCopies(Points, false);
+
+	FDoubleVector OldLocationOffset = LocationOffset;
 
 	// Initialize the Octree
 	Initialize(NewBounds);
@@ -768,7 +810,7 @@ void ULidarPointCloud::Merge(TArray<ULidarPointCloud*> PointCloudsToMerge, TFunc
 	}
 
 	// Re-insert original points
-	InsertPoints(Points, GetDefault<ULidarPointCloudSettings>()->DuplicateHandling, false);
+	InsertPoints(Points, GetDefault<ULidarPointCloudSettings>()->DuplicateHandling, false, (OldLocationOffset - LocationOffset).ToVector());
 
 	int32 Progress = 0;
 	float ProgressMulti = 1.0f / (1 + PointCloudsToMerge.Num());
@@ -782,8 +824,8 @@ void ULidarPointCloud::Merge(TArray<ULidarPointCloud*> PointCloudsToMerge, TFunc
 	for (ULidarPointCloud* Asset : PointCloudsToMerge)
 	{
 		// #todo: Do this in batches to avoid spiking RAM
-		Asset->GetPointsAsCopies(Points);
-		InsertPoints(Points, GetDefault<ULidarPointCloudSettings>()->DuplicateHandling, false);
+		Asset->GetPointsAsCopies(Points, false);
+		InsertPoints(Points, GetDefault<ULidarPointCloudSettings>()->DuplicateHandling, false, (Asset->LocationOffset - LocationOffset).ToVector());
 
 		if (ProgressCallback)
 		{
@@ -791,8 +833,6 @@ void ULidarPointCloud::Merge(TArray<ULidarPointCloud*> PointCloudsToMerge, TFunc
 		}
 	}
 
-	RefreshPointsBounds();
-
 	auto PostFunction = [this]() {
 		MarkPackageDirty();
 		OnPointCloudRebuiltEvent.Broadcast();
@@ -830,8 +870,8 @@ void ULidarPointCloud::AlignClouds(TArray<ULidarPointCloud*> PointCloudsToAlign)
 	// Calculate combined bounds
 	for (ULidarPointCloud* Asset : PointCloudsToAlign)
 	{
-		FDoubleVector OffsetMin = Asset->OriginalCoordinates + Asset->GetBounds().Min;
-		FDoubleVector OffsetMax = Asset->OriginalCoordinates + Asset->GetBounds().Max;
+		FDoubleVector OffsetMin = Asset->OriginalCoordinates - Asset->Octree.GetExtent();
+		FDoubleVector OffsetMax = Asset->OriginalCoordinates + Asset->Octree.GetExtent();
 
 		Min.X = FMath::Min(Min.X, OffsetMin.X);
 		Min.Y = FMath::Min(Min.Y, OffsetMin.Y);
@@ -847,9 +887,7 @@ void ULidarPointCloud::AlignClouds(TArray<ULidarPointCloud*> PointCloudsToAlign)
 	// Calculate and apply individual shifts
 	for (ULidarPointCloud* Asset : PointCloudsToAlign)
 	{
-		FDoubleVector OffsetCenter = Asset->OriginalCoordinates + Asset->GetBounds().GetCenter();
-		FDoubleVector Delta = OffsetCenter - CombinedCenter;
-		Asset->ShiftPointsBy(Delta, true);
+		Asset->SetLocationOffset(Asset->OriginalCoordinates - CombinedCenter);
 	}
 }
 
@@ -965,35 +1003,6 @@ void ULidarPointCloud::FinishPhysicsAsyncCook(bool bSuccess, UBodySetup* NewBody
 	bCollisionBuildInProgress = false;
 }
 
-class FPointCloudLatentAction : public FPendingLatentAction
-{
-public:
-	FName ExecutionFunction;
-	int32 OutputLink;
-	FWeakObjectPtr CallbackTarget;
-	ELidarPointCloudAsyncMode* Mode = nullptr;
-
-	FPointCloudLatentAction(const FLatentActionInfo& LatentInfo, ELidarPointCloudAsyncMode& Mode)
-		: ExecutionFunction(LatentInfo.ExecutionFunction)
-		, OutputLink(LatentInfo.Linkage)
-		, CallbackTarget(LatentInfo.CallbackTarget)
-		, Mode(&Mode)
-	{
-	}
-
-	virtual void UpdateOperation(FLatentResponse& Response) override
-	{
-		if (*Mode != ELidarPointCloudAsyncMode::Progress)
-		{
-			Response.FinishAndTriggerIf(true, ExecutionFunction, OutputLink, CallbackTarget);
-		}
-		else
-		{
-			Response.TriggerLink(ExecutionFunction, OutputLink, CallbackTarget);
-		}
-	}
-};
-
 void ULidarPointCloudBlueprintLibrary::CreatePointCloudFromFile(UObject* WorldContextObject, const FString& Filename, bool bUseAsync, FLatentActionInfo LatentInfo, ELidarPointCloudAsyncMode& AsyncMode, float& Progress, ULidarPointCloud*& PointCloud)
 {	
 	CreatePointCloudFromFile(WorldContextObject, Filename, bUseAsync, LatentInfo, FLidarPointCloudImportSettings::MakeGeneric(Filename), AsyncMode, Progress, PointCloud);
@@ -1089,7 +1098,7 @@ bool ULidarPointCloudBlueprintLibrary::LineTraceMulti(UObject* WorldContextObjec
 			auto Component = Actor->GetPointCloudComponent();
 
 			FLidarPointCloudTraceHit Hit(Actor, Component);
-			if (Component->LineTraceMulti(Ray, Radius, bVisibleOnly, Hit.Points))
+			if (Component->LineTraceMulti(Ray, Radius, bVisibleOnly, true, Hit.Points))
 			{
 				Hits.Add(Hit);
 				return true;
@@ -1100,9 +1109,9 @@ bool ULidarPointCloudBlueprintLibrary::LineTraceMulti(UObject* WorldContextObjec
 	return Hits.Num() > 0;
 }
 
-template bool ULidarPointCloud::InsertPoints_NoLock<FLidarPointCloudPoint*>(FLidarPointCloudPoint*, const int64&, ELidarPointCloudDuplicateHandling, bool, FThreadSafeBool*, TFunction<void(float)>);
-template bool ULidarPointCloud::InsertPoints_NoLock<const FLidarPointCloudPoint*>(const FLidarPointCloudPoint*, const int64&, ELidarPointCloudDuplicateHandling, bool, FThreadSafeBool*, TFunction<void(float)>);
-template bool ULidarPointCloud::InsertPoints_NoLock<FLidarPointCloudPoint**>(FLidarPointCloudPoint**, const int64&, ELidarPointCloudDuplicateHandling, bool, FThreadSafeBool*, TFunction<void(float)>);
+template bool ULidarPointCloud::InsertPoints_NoLock<FLidarPointCloudPoint*>(FLidarPointCloudPoint*, const int64&, ELidarPointCloudDuplicateHandling, bool, const FVector&, FThreadSafeBool*, TFunction<void(float)>);
+template bool ULidarPointCloud::InsertPoints_NoLock<const FLidarPointCloudPoint*>(const FLidarPointCloudPoint*, const int64&, ELidarPointCloudDuplicateHandling, bool, const FVector&, FThreadSafeBool*, TFunction<void(float)>);
+template bool ULidarPointCloud::InsertPoints_NoLock<FLidarPointCloudPoint**>(FLidarPointCloudPoint**, const int64&, ELidarPointCloudDuplicateHandling, bool, const FVector&, FThreadSafeBool*, TFunction<void(float)>);
 template bool ULidarPointCloud::SetData<const FLidarPointCloudPoint*>(const FLidarPointCloudPoint*, const int64&, TFunction<void(float)>);
 template bool ULidarPointCloud::SetData<FLidarPointCloudPoint**>(FLidarPointCloudPoint**, const int64&, TFunction<void(float)>);
 
diff --git a/Engine/Plugins/Enterprise/LidarPointCloud/Source/LidarPointCloudRuntime/Private/LidarPointCloudComponent.cpp b/Engine/Plugins/Enterprise/LidarPointCloud/Source/LidarPointCloudRuntime/Private/LidarPointCloudComponent.cpp
--- a/Engine/Plugins/Enterprise/LidarPointCloud/Source/LidarPointCloudRuntime/Private/LidarPointCloudComponent.cpp
+++ b/Engine/Plugins/Enterprise/LidarPointCloud/Source/LidarPointCloudRuntime/Private/LidarPointCloudComponent.cpp
@@ -19,8 +19,6 @@
 
 ULidarPointCloudComponent::ULidarPointCloudComponent()
 	: CustomMaterial(nullptr)
-	, PointBudget(1000000)
-	, ScreenCenterImportance(0.0f)
 	, MinScreenSize(0.05f)
 	, PointSize(1.0f)
 	, ColorSource(ELidarPointCloudColorationMode::Data)
@@ -39,7 +37,7 @@ ULidarPointCloudComponent::ULidarPointCloudComponent()
 	, MaxDepth(-1)
 	, bDrawNodeBounds(false)
 	, Material(nullptr)
-	, bOwnedByEditor(false)
+	, OwningViewportClient(nullptr)
 {
 	PrimaryComponentTick.bCanEverTick = false;
 	Mobility = EComponentMobility::Movable;
@@ -56,7 +54,7 @@ ULidarPointCloudComponent::ULidarPointCloudComponent()
 
 FBoxSphereBounds ULidarPointCloudComponent::CalcBounds(const FTransform& LocalToWorld) const
 {
-	return PointCloud ? PointCloud->GetPointsBounds().TransformBy(LocalToWorld) : USceneComponent::CalcBounds(LocalToWorld);
+	return PointCloud ? PointCloud->GetBounds().TransformBy(LocalToWorld) : USceneComponent::CalcBounds(LocalToWorld);
 }
 
 void ULidarPointCloudComponent::UpdateMaterial()
@@ -99,11 +97,14 @@ void ULidarPointCloudComponent::OnPointCloudRebuilt()
 	UpdateBounds();
 	UpdateMaterial();
 
-	if (ClassificationColors.Num() == 0)
+	if (PointCloud)
 	{
-		for (uint8& Classification : PointCloud->GetClassificationsImported())
+		if (ClassificationColors.Num() == 0)
 		{
-			ClassificationColors.Emplace(Classification, FLinearColor::White);
+			for (uint8& Classification : PointCloud->GetClassificationsImported())
+			{
+				ClassificationColors.Emplace(Classification, FLinearColor::White);
+			}
 		}
 	}
 }
diff --git a/Engine/Plugins/Enterprise/LidarPointCloud/Source/LidarPointCloudRuntime/Private/LidarPointCloudLODManager.cpp b/Engine/Plugins/Enterprise/LidarPointCloud/Source/LidarPointCloudRuntime/Private/LidarPointCloudLODManager.cpp
new file mode 100644
--- /dev/null
+++ b/Engine/Plugins/Enterprise/LidarPointCloud/Source/LidarPointCloudRuntime/Private/LidarPointCloudLODManager.cpp
@@ -0,0 +1,714 @@
+// Copyright Epic Games, Inc. All Rights Reserved.
+
+#include "LidarPointCloudLODManager.h"
+#include "LidarPointCloud.h"
+#include "LidarPointCloudComponent.h"
+#include "LidarPointCloudOctree.h"
+#include "LidarPointCloudSettings.h"
+#include "Rendering/LidarPointCloudRenderBuffers.h"
+#include "Engine/LocalPlayer.h"
+#include "Kismet/GameplayStatics.h"
+#include "Async/Async.h"
+#include "Misc/ScopeLock.h"
+#include "Engine/Engine.h"
+
+#if WITH_EDITOR
+#include "Classes/EditorStyleSettings.h"
+#include "EditorViewportClient.h"
+#endif
+
+DECLARE_CYCLE_STAT(TEXT("Buffer Creation"), STAT_BufferUpdate, STATGROUP_LidarPointCloud);
+DECLARE_CYCLE_STAT(TEXT("Buffer Update"), STAT_BufferUpdateRT, STATGROUP_LidarPointCloud);
+DECLARE_CYCLE_STAT(TEXT("Node Selection"), STAT_NodeSelection, STATGROUP_LidarPointCloud);
+DECLARE_DWORD_COUNTER_STAT(TEXT("Registered Proxies"), STAT_ProxyCount, STATGROUP_LidarPointCloud)
+DECLARE_DWORD_COUNTER_STAT(TEXT("Visible Nodes"), STAT_NodeCount, STATGROUP_LidarPointCloud)
+DECLARE_DWORD_COUNTER_STAT(TEXT("Visible Points"), STAT_PointCount, STATGROUP_LidarPointCloud)
+DECLARE_DWORD_COUNTER_STAT(TEXT("Total Point Count [thousands]"), STAT_PointCountTotal, STATGROUP_LidarPointCloud)
+
+static TAutoConsoleVariable<int32> CVarLidarPointBudget(
+	TEXT("r.LidarPointBudget"),
+	1000000,
+	TEXT("Determines the maximum number of points to be visible on the screen.\n")
+	TEXT("Higher values will produce better image quality, but will require faster hardware."),
+	ECVF_Scalability);
+
+static TAutoConsoleVariable<float> CVarLidarScreenCenterImportance(
+	TEXT("r.LidarScreenCenterImportance"),
+	0.0f,
+	TEXT("Determines the preference towards selecting nodes closer to screen center\n")
+	TEXT("with larger values giving more priority towards screen center.\n")
+	TEXT("Usefulf for VR, where edge vision is blurred anyway.\n")
+	TEXT("0 to disable."),
+	ECVF_Scalability);
+
+static TAutoConsoleVariable<float> CVarBaseLODImportance(
+	TEXT("r.LidarBaseLODImportance"),
+	0.1f,
+	TEXT("Determines the importance of selecting at least the base LOD of far assets.\n")
+	TEXT("Increase it, if you're experiencing actor 'popping'.\n")
+	TEXT("0 to use purely screensize-driven algorithm."),
+	ECVF_Scalability);
+
+/**
+ * Iterates over the provided nodes and sets location and color data.
+ * Returns the total number of points processed.
+ */
+uint32 SetLocationAndColorData(uint8* Buffer, const TArray<FLidarPointCloudTraversalOctreeNode*>& Nodes, bool bUseClassification, bool OwningViewportClient)
+{
+	FColor SelectionColor = FColor::White;
+
+#if WITH_EDITOR
+	SelectionColor = GetDefault<UEditorStyleSettings>()->SelectionColor.ToFColor(false);
+#endif
+
+	uint8* BufferCurrent = Buffer;
+
+	if (OwningViewportClient)
+	{
+		if (bUseClassification)
+		{
+			for (auto& Node : Nodes)
+			{
+				// Skip nodes with no available data
+				if (!Node->DataNode->HasData())
+				{
+					continue;
+				}
+
+				for (FLidarPointCloudPoint* Data = Node->DataNode->GetData(), *DataEnd = Data + Node->DataNode->GetNumVisiblePoints(); Data != DataEnd; ++Data)
+				{
+					FMemory::Memcpy(BufferCurrent, Data, 12);
+					BufferCurrent += 12;
+
+					if (Data->bSelected)
+					{
+						FMemory::Memcpy(BufferCurrent, &SelectionColor, 4);
+					}
+					else
+					{
+						FColor ClassificationColor(Data->ClassificationID, Data->ClassificationID, Data->ClassificationID, Data->Color.A);
+						FMemory::Memcpy(BufferCurrent, &ClassificationColor, 4);
+					}
+
+					BufferCurrent += 4;
+				}
+			}
+		}
+		else
+		{
+			for (auto& Node : Nodes)
+			{
+				// Skip nodes with no available data
+				if (!Node->DataNode->HasData())
+				{
+					continue;
+				}
+
+				for (FLidarPointCloudPoint* Data = Node->DataNode->GetData(), *DataEnd = Data + Node->DataNode->GetNumVisiblePoints(); Data != DataEnd; ++Data)
+				{
+					if (Data->bSelected)
+					{
+						FMemory::Memcpy(BufferCurrent, Data, 12);
+						BufferCurrent += 12;
+
+						FMemory::Memcpy(BufferCurrent, &SelectionColor, 4);
+						BufferCurrent += 4;
+					}
+					else
+					{
+						FMemory::Memcpy(BufferCurrent, Data, 16);
+						BufferCurrent += 16;
+					}
+				}
+			}
+		}
+	}
+	else
+	{
+		if (bUseClassification)
+		{
+			for (auto& Node : Nodes)
+			{
+				// Skip nodes with no available data
+				if (!Node->DataNode->HasData())
+				{
+					continue;
+				}
+
+				for (FLidarPointCloudPoint* Data = Node->DataNode->GetData(), *DataEnd = Data + Node->DataNode->GetNumVisiblePoints(); Data != DataEnd; ++Data)
+				{
+					FMemory::Memcpy(BufferCurrent, Data, 12);
+					BufferCurrent += 12;
+
+					FColor ClassificationColor(Data->ClassificationID, Data->ClassificationID, Data->ClassificationID, Data->Color.A);
+					FMemory::Memcpy(BufferCurrent, &ClassificationColor, 4);
+					BufferCurrent += 4;
+				}
+			}
+		}
+		else
+		{
+			for (auto& Node : Nodes)
+			{
+				// Skip nodes with no available data
+				if (!Node->DataNode->HasData())
+				{
+					continue;
+				}
+
+				for (FLidarPointCloudPoint* Data = Node->DataNode->GetData(), *DataEnd = Data + Node->DataNode->GetNumVisiblePoints(); Data != DataEnd; ++Data)
+				{
+					FMemory::Memcpy(BufferCurrent, Data, 16);
+					BufferCurrent += 16;
+				}
+			}
+		}
+	}
+
+	// Calculates the actual number of instances copied to the buffer (accounts for the invisible points)
+	return (BufferCurrent - Buffer) / 16;
+}
+
+/** Iterates over the provided nodes and sets scale data */
+void SetScaleData(uint8* Buffer, const TArray<FLidarPointCloudTraversalOctreeNode*>& Nodes)
+{
+	for (auto Node : Nodes)
+	{
+		// Skip nodes with no available data
+		if (!Node->DataNode->HasData())
+		{
+			continue;
+		}
+
+		FMemory::Memset(Buffer, Node->VirtualDepth, Node->DataNode->GetNumVisiblePoints());
+		Buffer += Node->DataNode->GetNumVisiblePoints();
+	}
+}
+
+FLidarPointCloudViewData::FLidarPointCloudViewData(bool bCompute)
+	: bValid(false)
+	, ViewOrigin(FVector::ZeroVector)
+	, ViewDirection(FVector::ForwardVector)
+	, ScreenSizeFactor(0)
+	, bSkipMinScreenSize(false)
+	, bPIE(false)
+	, bHasFocus(false)
+{
+	if (bCompute)
+	{
+		Compute();
+	}
+}
+
+void FLidarPointCloudViewData::Compute()
+{
+#if WITH_EDITOR
+	bPIE = false;
+	if (GIsEditor && GEditor && GEditor->GetActiveViewport())
+	{
+		bPIE = GEditor->GetActiveViewport() == GEditor->GetPIEViewport();
+		
+		// PIE needs a different computation method
+		if (!bPIE)
+		{
+			ComputeFromEditorViewportClient(GEditor->GetActiveViewport()->GetClient());
+		}
+
+		// Simulating counts as PIE for the purpose of LOD calculation
+		bPIE |= GEditor->bIsSimulatingInEditor;
+	}
+#endif
+
+	// Check for bValid, in case we run inside editor and already have the View set up
+	if (!bValid && GEngine)
+	{
+		ULocalPlayer* const LP = GEngine->FindFirstLocalPlayerFromControllerId(0);
+		if (LP && LP->ViewportClient)
+		{
+			FSceneViewProjectionData ProjectionData;
+			if (LP->GetProjectionData(LP->ViewportClient->Viewport, eSSP_FULL, ProjectionData))
+			{
+				ViewOrigin = ProjectionData.ViewOrigin;
+				FMatrix ViewRotationMatrix = ProjectionData.ViewRotationMatrix;
+				if (!ViewRotationMatrix.GetOrigin().IsNearlyZero(0.0f))
+				{
+					ViewOrigin += ViewRotationMatrix.InverseTransformPosition(FVector::ZeroVector);
+					ViewRotationMatrix = ViewRotationMatrix.RemoveTranslation();
+				}
+
+				FMatrix ViewMatrix = FTranslationMatrix(-ViewOrigin) * ViewRotationMatrix;
+				ViewDirection = ViewMatrix.GetColumn(2);
+				FMatrix ProjectionMatrix = AdjustProjectionMatrixForRHI(ProjectionData.ProjectionMatrix);
+
+				ScreenSizeFactor = FMath::Square(FMath::Max(0.5f * ProjectionMatrix.M[0][0], 0.5f * ProjectionMatrix.M[1][1]));
+
+				// Skip SS check, if not in the projection view nor game world
+				bSkipMinScreenSize = (ProjectionMatrix.M[3][3] >= 1.0f) && !LP->GetWorld()->IsGameWorld();
+				GetViewFrustumBounds(ViewFrustum, ViewMatrix * ProjectionMatrix, false);
+
+				bHasFocus = LP->ViewportClient->Viewport->HasFocus();
+
+				bValid = true;
+			}
+		}
+	}
+}
+
+bool FLidarPointCloudViewData::ComputeFromEditorViewportClient(FViewportClient* ViewportClient)
+{
+#if WITH_EDITOR
+	if (FEditorViewportClient* Client = (FEditorViewportClient*)ViewportClient)
+	{
+		if (Client->Viewport && Client->Viewport->GetSizeXY() != FIntPoint::ZeroValue)
+		{
+			FSceneViewFamily::ConstructionValues CVS(nullptr, nullptr, FEngineShowFlags(EShowFlagInitMode::ESFIM_Game));
+			CVS.SetWorldTimes(0, 0, 0);
+			FSceneViewFamily ViewFamily(CVS);
+			FSceneView* View = Client->CalcSceneView(&ViewFamily);
+
+			const FMatrix& ProjectionMatrix = View->ViewMatrices.GetProjectionMatrix();
+			ScreenSizeFactor = FMath::Square(FMath::Max(0.5f * ProjectionMatrix.M[0][0], 0.5f * ProjectionMatrix.M[1][1]));
+			ViewOrigin = View->ViewMatrices.GetViewOrigin();
+			ViewDirection = View->GetViewDirection();
+			ViewFrustum = View->ViewFrustum;
+			bSkipMinScreenSize = !View->bIsGameView && !View->IsPerspectiveProjection();
+			bHasFocus = Client->Viewport->HasFocus();
+
+			bValid = true;
+
+			return true;
+		}
+	}
+#endif
+	return false;
+}
+
+void FLidarPointCloudTraversalOctree::GetVisibleNodes(TArray<FLidarPointCloudLODManager::FNodeSizeData>& NodeSizeData, const FLidarPointCloudViewData& ViewData, const int32& ProxyIndex, const FLidarPointCloudNodeSelectionParams& SelectionParams, const float& CurrentTime)
+{
+	// Skip processing, if the asset is not visible
+	if (!ViewData.ViewFrustum.IntersectBox(GetCenter(), Octree->GetExtent()))
+	{
+		return;
+	}
+
+	float MinScreenSizeSq = SelectionParams.MinScreenSize * SelectionParams.MinScreenSize;
+	float BoundsScaleSq = SelectionParams.BoundsScale * SelectionParams.BoundsScale;
+
+	const float BaseLODImportance = FMath::Max(0.0f, CVarBaseLODImportance.GetValueOnAnyThread());
+
+	TQueue<FLidarPointCloudTraversalOctreeNode*> Nodes;
+	FLidarPointCloudTraversalOctreeNode* CurrentNode = nullptr;
+	Nodes.Enqueue(&Root);
+	while (Nodes.Dequeue(CurrentNode))
+	{
+		// Reset selection flag
+		CurrentNode->bSelected = false;
+
+		// Update number of visible points, if needed
+		CurrentNode->DataNode->UpdateNumVisiblePoints();
+
+		// In Frustum?
+		// #todo: Skip frustum checks for nodes fully in frustum
+		if (!ViewData.ViewFrustum.IntersectBox(CurrentNode->Center, Extents[CurrentNode->Depth] * SelectionParams.BoundsScale))
+		{
+			continue;
+		}
+
+		// Only process this node if it has any visible points - do not use continue; as the children may still contain visible points!
+		if (CurrentNode->DataNode->GetNumVisiblePoints() > 0 && CurrentNode->Depth >= SelectionParams.MinDepth)
+		{
+			float ScreenSizeSq = 0;
+
+			FVector VectorToNode = CurrentNode->Center - ViewData.ViewOrigin;
+			const float DistSq = VectorToNode.SizeSquared();
+			const float AdjustedRadiusSq = RadiiSq[CurrentNode->Depth] * BoundsScaleSq;
+
+			// Make sure to show at least the minimum depth for each visible asset
+			if (CurrentNode->Depth == SelectionParams.MinDepth)
+			{
+				// Add screen size to maintain hierarchy
+				ScreenSizeSq = BaseLODImportance + ViewData.ScreenSizeFactor * AdjustedRadiusSq / FMath::Max(1.0f, DistSq);
+			}
+			else
+			{
+				// If the camera is within this node's bounds, it should always be qualified for rendering
+				if (DistSq <= AdjustedRadiusSq)
+				{
+					// Subtract Depth to maintain hierarchy 
+					ScreenSizeSq = 1000 - CurrentNode->Depth;
+				}
+				else
+				{
+					ScreenSizeSq = ViewData.ScreenSizeFactor * AdjustedRadiusSq / FMath::Max(1.0f, DistSq);
+
+					// Check for minimum screen size
+					if (!ViewData.bSkipMinScreenSize && ScreenSizeSq < MinScreenSizeSq)
+					{
+						continue;
+					}
+
+					// Add optional preferential selection for nodes closer to the screen center
+					if (SelectionParams.ScreenCenterImportance > 0)
+					{
+						VectorToNode.Normalize();
+						float Dot = FVector::DotProduct(ViewData.ViewDirection, VectorToNode);
+
+						ScreenSizeSq = FMath::Lerp(ScreenSizeSq, ScreenSizeSq * Dot, SelectionParams.ScreenCenterImportance);
+					}
+				}
+			}
+
+			NodeSizeData.Emplace(CurrentNode, ScreenSizeSq, ProxyIndex);
+		}
+
+		if (SelectionParams.MaxDepth < 0 || CurrentNode->Depth < SelectionParams.MaxDepth)
+		{
+			for (auto& Child : CurrentNode->Children)
+			{
+				Nodes.Enqueue(&Child);
+			}
+		}
+	}
+}
+
+void FLidarPointCloudLODManager::Tick(float DeltaTime)
+{
+	// Skip processing, if a previous one is still going
+	if (bProcessing)
+	{
+		return;
+	}
+
+	bProcessing = true;
+
+	Time += DeltaTime;
+	
+	PrepareProxies();
+
+	// A copy of the array will be passed, to avoid concurrency issues
+	TArray<FRegisteredProxy> CurrentRegisteredProxies = RegisteredProxies;
+
+	Async(EAsyncExecution::ThreadPool, [this, CurrentRegisteredProxies] { ProcessLOD(CurrentRegisteredProxies, Time); });
+}
+
+TStatId FLidarPointCloudLODManager::GetStatId() const
+{
+	RETURN_QUICK_DECLARE_CYCLE_STAT(LidarPointCloudLODManager, STATGROUP_Tickables);
+}
+
+void FLidarPointCloudLODManager::RegisterProxy(ULidarPointCloudComponent* Component, TWeakPtr<FLidarPointCloudSceneProxyWrapper, ESPMode::ThreadSafe> SceneProxyWrapper)
+{
+	if (IsValid(Component))
+	{
+		static FLidarPointCloudLODManager Instance;
+		Instance.RegisteredProxies.Emplace(Component, SceneProxyWrapper);
+	}
+}
+
+void FLidarPointCloudLODManager::ProcessLOD(const TArray<FLidarPointCloudLODManager::FRegisteredProxy>& InRegisteredProxies, const float CurrentTime)
+{
+	const int32 PointBudget = CVarLidarPointBudget.GetValueOnAnyThread();
+	static FLidarPointCloudDataBufferManager BufferManager(PointBudget * 17);
+	BufferManager.Resize(PointBudget * 17);
+
+	int32 TotalPointsSelected = 0;
+
+	TArray<TArray<FLidarPointCloudTraversalOctreeNode*>> SelectedNodesData;
+
+	// Node selection
+	{
+		SCOPE_CYCLE_COUNTER(STAT_NodeSelection);
+
+		const float ScreenCenterImportance = CVarLidarScreenCenterImportance.GetValueOnAnyThread();
+
+		int32 NumSelectedNodes = 0;
+
+		TArray<FNodeSizeData> NodeSizeData;
+
+		for (int32 i = 0; i < InRegisteredProxies.Num(); ++i)
+		{
+			auto& RegisteredProxy = InRegisteredProxies[i];
+
+			// Acquire a Shared Pointer from the Weak Pointer and check that it references a valid object
+			if (auto SceneProxyWrapper = RegisteredProxy.SceneProxyWrapper.Pin())
+			{
+				FScopeLock OctreeLock(&RegisteredProxy.PointCloud->Octree.DataLock);
+
+				// If the octree has been invalidated, skip processing
+				if (!RegisteredProxy.TraversalOctree->bValid)
+				{
+					continue;
+				}
+
+#if WITH_EDITOR
+				// Avoid doubling the point allocation of the same asset (once in Editor world and once in PIE world)
+				if (RegisteredProxy.bSkip)
+				{
+					continue;
+				}
+#endif
+
+				// Construct selection params
+				FLidarPointCloudNodeSelectionParams SelectionParams;
+				SelectionParams.MinScreenSize = FMath::Max(RegisteredProxy.Component->MinScreenSize, 0.0f);
+				SelectionParams.ScreenCenterImportance = ScreenCenterImportance;
+				SelectionParams.MinDepth = RegisteredProxy.Component->MinDepth;
+				SelectionParams.MaxDepth = RegisteredProxy.Component->MaxDepth;
+				SelectionParams.BoundsScale = RegisteredProxy.Component->BoundsScale;
+
+				// Append visible nodes
+				RegisteredProxy.TraversalOctree->GetVisibleNodes(NodeSizeData, RegisteredProxy.ViewData, i, SelectionParams, CurrentTime);
+			}
+		}
+
+		// Sort Nodes
+		Algo::Sort(NodeSizeData, [](const FNodeSizeData& A, const FNodeSizeData& B) { return A.Size > B.Size; });
+
+		// Limit nodes using specified Point Budget
+		SelectedNodesData.AddDefaulted(InRegisteredProxies.Num());
+		for (FNodeSizeData& Element : NodeSizeData)
+		{
+			int32 NewNumPointsSelected = TotalPointsSelected + Element.Node->DataNode->GetNumVisiblePoints();
+
+			if (NewNumPointsSelected <= PointBudget)
+			{
+				SelectedNodesData[Element.ProxyIndex].Add(Element.Node);
+				TotalPointsSelected = NewNumPointsSelected;
+				Element.Node->bSelected = true;
+				++NumSelectedNodes;
+			}
+		}
+
+		INC_DWORD_STAT_BY(STAT_PointCount, TotalPointsSelected);
+		INC_DWORD_STAT_BY(STAT_NodeCount, NumSelectedNodes);
+	}
+
+	// Used to pass render data updates to render thread
+	auto Buffer = BufferManager.GetFreeBuffer();
+	TArray<FLidarPointCloudProxyUpdateData> ProxyUpdateData;
+
+	// Build buffer data
+	{
+		SCOPE_CYCLE_COUNTER(STAT_BufferUpdate);
+
+		// Prepare the data for structured buffer.
+		uint8* BufferData = Buffer->GetData();
+		uint8* LocationAndColorBufferPtr = BufferData;
+		uint8* ScaleBufferPtr = BufferData + TotalPointsSelected * 16;
+
+		int32 FirstElementIndex = 0;
+
+		// Set when to release the BulkData, if no longer visible
+		const float BulkDataLifetime = CurrentTime + 1;
+
+		for (int32 i = 0; i < SelectedNodesData.Num(); ++i)
+		{
+			auto& RegisteredProxy = InRegisteredProxies[i];
+
+			// Only calculate if needed
+			if (RegisteredProxy.Component->PointSize > 0)
+			{
+				for (auto Node : SelectedNodesData[i])
+				{
+					Node->CalculateVirtualDepth(RegisteredProxy.TraversalOctree->LevelWeights, RegisteredProxy.TraversalOctree->VirtualDepthMultiplier, RegisteredProxy.Component->PointSizeBias);
+				}
+			}
+
+			int32 NumPoints = 0;
+
+			// Since the process is async, make sure we can access the data!
+			{
+				FScopeLock OctreeLock(&RegisteredProxy.PointCloud->Octree.DataLock);
+				
+				// If the octree has been invalidated, skip processing
+				if (!RegisteredProxy.TraversalOctree->bValid)
+				{
+					continue;
+				}
+
+				// Queue nodes to be streamed
+				for (auto Node : SelectedNodesData[i])
+				{
+					RegisteredProxy.PointCloud->Octree.QueueNode(Node->DataNode, BulkDataLifetime);
+				}
+
+				NumPoints = SetLocationAndColorData(LocationAndColorBufferPtr, SelectedNodesData[i], RegisteredProxy.Component->ColorSource == ELidarPointCloudColorationMode::Classification, RegisteredProxy.Component->IsOwnedByEditor());
+				SetScaleData(ScaleBufferPtr, SelectedNodesData[i]);
+			}
+
+			FLidarPointCloudProxyUpdateData UpdateData;
+			UpdateData.SceneProxyWrapper = RegisteredProxy.SceneProxyWrapper;
+			UpdateData.FirstElementIndex = FirstElementIndex;
+			UpdateData.NumElements = NumPoints;
+			UpdateData.PointBudget = PointBudget;
+			UpdateData.VDMultiplier = RegisteredProxy.TraversalOctree->ReversedVirtualDepthMultiplier;
+			UpdateData.RootCellSize = RegisteredProxy.PointCloud->Octree.GetRootCellSize();
+
+#if !(UE_BUILD_SHIPPING)
+			// Prepare bounds
+			if (RegisteredProxy.Component->bDrawNodeBounds)
+			{
+				UpdateData.Bounds.Reset(SelectedNodesData[i].Num());
+
+				for (auto Node : SelectedNodesData[i])
+				{
+					FVector Extent = RegisteredProxy.TraversalOctree->Extents[Node->Depth];
+					UpdateData.Bounds.Emplace(Node->Center - Extent, Node->Center + Extent);
+				}
+			}
+#endif
+
+			ProxyUpdateData.Add(UpdateData);
+
+			// Shift pointers
+			FirstElementIndex += NumPoints;
+			LocationAndColorBufferPtr += NumPoints * 16;
+			ScaleBufferPtr += NumPoints;
+		}
+	}
+
+	// Begin streaming data
+	for (int32 i = 0; i < InRegisteredProxies.Num(); ++i)
+	{
+		auto& RegisteredProxy = InRegisteredProxies[i];
+
+		FScopeLock OctreeLock(&RegisteredProxy.PointCloud->Octree.DataLock);
+		RegisteredProxy.PointCloud->Octree.UnloadOldNodes(CurrentTime);
+		RegisteredProxy.PointCloud->Octree.StreamQueuedNodes();
+	}
+
+	// Process buffer updates on RT
+	ENQUEUE_RENDER_COMMAND(ProcessLidarPointCloudLOD)([PointBudget, Buffer, TotalPointsSelected, ProxyUpdateData](FRHICommandListImmediate& RHICmdList)
+	{
+		SCOPE_CYCLE_COUNTER(STAT_BufferUpdateRT);
+
+		// Resize IndexBuffer
+		GLidarPointCloudIndexBuffer.Resize(PointBudget);
+
+		// 17 bytes per point, element size set to 4 bytes to minimize wastage. Rounded to 4.3 elements per point
+		GLidarPointCloudRenderBuffer.Resize(PointBudget * 4.3f);
+
+		GLidarPointCloudRenderBuffer.PointCount = TotalPointsSelected;
+
+		if (TotalPointsSelected > 0)
+		{
+			int32 TotalDataSize = TotalPointsSelected * 17;
+
+			// Update contents of the Structured Buffer
+			uint8* StructuredBuffer = (uint8*)RHILockVertexBuffer(GLidarPointCloudRenderBuffer.Buffer, 0, TotalDataSize, RLM_WriteOnly);
+			FMemory::Memcpy(StructuredBuffer, Buffer->GetData(), TotalDataSize);
+			RHIUnlockVertexBuffer(GLidarPointCloudRenderBuffer.Buffer);
+
+			// Iterate over proxies and, if valid, update its FirstElementIndex
+			for (int32 i = 0; i < ProxyUpdateData.Num(); ++i)
+			{
+				// Check for proxy's validity, in case it has been destroyed since the update was issued
+				if (TSharedPtr<FLidarPointCloudSceneProxyWrapper, ESPMode::ThreadSafe> SceneProxyWrapper = ProxyUpdateData[i].SceneProxyWrapper.Pin())
+				{
+					SceneProxyWrapper->Proxy->UpdateRenderData(ProxyUpdateData[i]);
+				}
+			}
+		}
+
+		Buffer->MarkAsFree();
+	});
+
+	bProcessing = false;
+}
+
+void FLidarPointCloudLODManager::PrepareProxies()
+{
+	FLidarPointCloudViewData ViewData(true);
+
+	const bool bPrioritizeActiveViewport = GetDefault<ULidarPointCloudSettings>()->bPrioritizeActiveViewport;
+
+	// Contains the total number of points contained by all assets (including invisible and culled)
+	int64 TotalPointCount = 0;
+
+	// Prepare proxies
+	for (int32 i = 0; i < RegisteredProxies.Num(); ++i)
+	{
+		FRegisteredProxy& RegisteredProxy = RegisteredProxies[i];
+		bool bValidProxy = false;
+
+		if (RegisteredProxy.Component->GetPointCloud())
+		{
+			// Acquire a Shared Pointer from the Weak Pointer and check that it references a valid object
+			if (auto SceneProxyWrapper = RegisteredProxy.SceneProxyWrapper.Pin())
+			{
+#if WITH_EDITOR
+				// Avoid doubling the point allocation of the same asset (once in Editor world and once in PIE world)
+				RegisteredProxy.bSkip = ViewData.bPIE && RegisteredProxy.Component->GetWorld()->WorldType == EWorldType::Type::Editor;
+#endif
+
+				// Check if the component's transform has changed, and invalidate the Traversal Octree if so
+				const FTransform Transform = RegisteredProxy.Component->GetComponentTransform();
+				if (!RegisteredProxy.LastComponentTransform.Equals(Transform))
+				{
+					RegisteredProxy.TraversalOctree->bValid = false;
+					RegisteredProxy.LastComponentTransform = Transform;
+				}
+
+				// Re-initialize the traversal octree, if needed
+				if (!RegisteredProxy.TraversalOctree->bValid)
+				{
+					// Update asset reference
+					RegisteredProxy.PointCloud = RegisteredProxy.Component->GetPointCloud();
+
+					// Recreate the Traversal Octree
+					RegisteredProxy.TraversalOctree = MakeShareable(new FLidarPointCloudTraversalOctree(&RegisteredProxy.PointCloud->Octree, RegisteredProxy.Component->GetComponentTransform()));
+					RegisteredProxy.PointCloud->Octree.RegisterTraversalOctree(RegisteredProxy.TraversalOctree);
+				}
+
+				// If this is an editor component, use its own ViewportClient
+				if (auto Client = RegisteredProxy.Component->GetOwningViewportClient().Pin())
+				{
+					// If the ViewData cannot be successfully retrieved from the editor viewport, fall back to using main view
+					if (!RegisteredProxy.ViewData.ComputeFromEditorViewportClient(Client.Get()))
+					{
+						RegisteredProxy.ViewData = ViewData;
+					}
+				}
+				// ... otherwise, use the ViewData provided
+				else
+				{
+					RegisteredProxy.ViewData = ViewData;
+				}
+
+				// Increase priority, if the viewport has focus
+				if (bPrioritizeActiveViewport && RegisteredProxy.ViewData.bHasFocus)
+				{
+					RegisteredProxy.ViewData.ScreenSizeFactor *= 6;
+				}
+
+				// Don't count the skippable proxies
+				if (!RegisteredProxy.bSkip)
+				{
+					TotalPointCount += RegisteredProxy.PointCloud->GetNumPoints();
+				}
+
+				bValidProxy = true;
+			}
+		}
+		
+		// If the SceneProxy has been destroyed, remove it from the list and reiterate
+		if(!bValidProxy)
+		{
+			RegisteredProxies.RemoveAtSwap(i--, 1, false);
+		}
+	}
+
+	INC_DWORD_STAT_BY(STAT_PointCountTotal, TotalPointCount / 1000);
+	INC_DWORD_STAT_BY(STAT_ProxyCount, RegisteredProxies.Num());
+}
+
+FLidarPointCloudLODManager::FRegisteredProxy::FRegisteredProxy(ULidarPointCloudComponent* Component, TWeakPtr<FLidarPointCloudSceneProxyWrapper, ESPMode::ThreadSafe> SceneProxyWrapper)
+	: Component(Component)
+	, PointCloud(Component->GetPointCloud())
+	, SceneProxyWrapper(SceneProxyWrapper)
+	, TraversalOctree(new FLidarPointCloudTraversalOctree(&PointCloud->Octree, Component->GetComponentTransform()))
+	, LastComponentTransform(Component->GetComponentTransform())
+	, bSkip(false)
+{
+	PointCloud->Octree.RegisterTraversalOctree(TraversalOctree);
+}
