Merging

diff --git a/Engine/Plugins/Enterprise/LidarPointCloud/Source/LidarPointCloudRuntime/Private/LidarPointCloudOctree.cpp b/Engine/Plugins/Enterprise/LidarPointCloud/Source/LidarPointCloudRuntime/Private/LidarPointCloudOctree.cpp
--- a/Engine/Plugins/Enterprise/LidarPointCloud/Source/LidarPointCloudRuntime/Private/LidarPointCloudOctree.cpp
+++ b/Engine/Plugins/Enterprise/LidarPointCloud/Source/LidarPointCloudRuntime/Private/LidarPointCloudOctree.cpp
@@ -6,11 +6,10 @@
 #include "Collision/LidarPointCloudCollision.h"
 #include "Misc/ScopeTryLock.h"
 #include "Containers/Queue.h"
-#include "SceneView.h"
 #include "Async/Async.h"
 #include "Misc/FileHelper.h"
 
-DECLARE_CYCLE_STAT(TEXT("Node Selection"), STAT_NodeSelection, STATGROUP_LidarPointCloud);
+DECLARE_CYCLE_STAT(TEXT("Node Streaming"), STAT_NodeStreaming, STATGROUP_LidarPointCloud);
 
 int32 FLidarPointCloudOctree::MaxNodeDepth = (1 << (sizeof(FLidarPointCloudOctreeNode::Depth) * 8)) - 1;
 int32 FLidarPointCloudOctree::MaxBucketSize = 200;
@@ -37,19 +36,9 @@ struct FGridAllocation
 	}
 };
 
-/** Used for node size sorting and node selection. */
-struct FNodeSizeData
+FGridAllocation CalculateGridCellData(const FVector& Location, const FVector& Center, const FLidarPointCloudOctree::FSharedLODData& LODData)
 {
-	FLidarPointCloudTraversalOctreeNode* Node;
-	float Size;
-	FNodeSizeData(FLidarPointCloudTraversalOctreeNode* Node, const float& Size) : Node(Node), Size(Size) {}
-};
-
-auto PointVisibilitySortingFunction = [](const FLidarPointCloudPoint& A, const FLidarPointCloudPoint& B) { return A.bVisible > B.bVisible; };
-
-FGridAllocation CalculateGridCellData(const FLidarPointCloudPoint* Point, const FVector& Center, const FLidarPointCloudOctree::FSharedLODData& LODData)
-{
-	FVector CenterRelativeLocation = Point->Location - Center;
+	FVector CenterRelativeLocation = Location - Center;
 	FVector OffsetLocation = CenterRelativeLocation + LODData.Extent;
 	FVector NormalizedGridLocation = OffsetLocation * LODData.NormalizationMultiplier;
 
@@ -86,12 +75,15 @@ FLidarPointCloudOctree::FSharedLODData::FSharedLODData(const FVector& InExtent)
 
 //////////////////////////////////////////////////////////// FLidarPointCloudOctreeNode
 
-FLidarPointCloudOctreeNode::FLidarPointCloudOctreeNode(FLidarPointCloudOctree* Tree, FLidarPointCloudOctreeNode* Parent, const uint8& Depth, const uint8& LocationInParent, const FVector& Center)
-	: Depth(Depth)
+FLidarPointCloudOctreeNode::FLidarPointCloudOctreeNode(FLidarPointCloudOctree* Tree, const uint8& Depth, const uint8& LocationInParent, const FVector& Center)
+	: BulkDataLifetime(0)
+	, Depth(Depth)
 	, LocationInParent(LocationInParent)
 	, Center(Center)
 	, bVisibilityDirty(false)
 	, NumVisiblePoints(0)
+	, bHasDataPending(false)
+	, bCanReleaseData(true)
 {
 	if (Tree)
 	{
@@ -113,22 +105,11 @@ void FLidarPointCloudOctreeNode::UpdateNumVisiblePoints()
 	if (bVisibilityDirty)
 	{
 		// Sort points to speed up rendering
-		Algo::Sort(AllocatedPoints, MoveTemp(PointVisibilitySortingFunction));
-		Algo::Sort(PaddingPoints, MoveTemp(PointVisibilitySortingFunction));
+		SortVisiblePoints();
 
 		// Recalculate visibility
 		NumVisiblePoints = 0;
-		FOR(Point, AllocatedPoints)
-		{
-			if (!Point->bVisible)
-			{
-				break;
-			}
-
-			NumVisiblePoints++;
-		}
-
-		FOR(Point, PaddingPoints)
+		FOR_RO(Point, this)
 		{
 			if (!Point->bVisible)
 			{
@@ -142,6 +123,14 @@ void FLidarPointCloudOctreeNode::UpdateNumVisiblePoints()
 	}
 }
 
+FLidarPointCloudPoint* FLidarPointCloudOctreeNode::GetPersistentData() const
+{
+	FLidarPointCloudOctreeNode* mutable_this = const_cast<FLidarPointCloudOctreeNode*>(this);
+	mutable_this->bCanReleaseData = false;
+
+	return BulkData.GetData();
+}
+
 FBox FLidarPointCloudOctreeNode::GetBounds(const FLidarPointCloudOctree* Tree) const
 {
 	return FBox(Center - Tree->SharedData[Depth].Extent, Center + Tree->SharedData[Depth].Extent);
@@ -165,7 +154,7 @@ FLidarPointCloudOctreeNode* FLidarPointCloudOctreeNode::GetChildNodeAtLocation(c
 	return nullptr;
 }
 
-void FLidarPointCloudOctreeNode::InsertPoints(FLidarPointCloudOctree* Tree, const FLidarPointCloudPoint* Points, const int32& Count, ELidarPointCloudDuplicateHandling DuplicateHandling)
+void FLidarPointCloudOctreeNode::InsertPoints(FLidarPointCloudOctree* Tree, const FLidarPointCloudPoint* Points, const int32& Count, ELidarPointCloudDuplicateHandling DuplicateHandling, const FVector& Translation)
 {
 	const auto& LODData = Tree->SharedData[Depth];
 
@@ -180,7 +169,8 @@ void FLidarPointCloudOctreeNode::InsertPoints(FLidarPointCloudOctree* Tree, cons
 	// Filter the local set of incoming data
 	for (int32 Index = 0; Index < Count; ++Index)
 	{
-		auto InGridData = CalculateGridCellData(&Points[Index], Center, LODData);
+		const FVector AdjustedLocation = Points[Index].Location + Translation;
+		auto InGridData = CalculateGridCellData(AdjustedLocation, Center, LODData);
 		FGridAllocation* GridCell = NewGridAllocationMap.Find(InGridData.Index);
 
 		// Attempt to allocate the point to this node
@@ -204,7 +194,8 @@ void FLidarPointCloudOctreeNode::InsertPoints(FLidarPointCloudOctree* Tree, cons
 			{
 				if (bStoreInBucket)
 				{
-					PointBuckets[GridCell->ChildNodeLocation].Add(Points[GridCell->Index]);
+					const FLidarPointCloudPoint& Other = Points[GridCell->Index];
+					PointBuckets[GridCell->ChildNodeLocation].Emplace(AdjustedLocation, Other.Color, !!Other.bVisible, Other.ClassificationID);
 				}
 				
 				GridCell->Index = Index;
@@ -212,7 +203,8 @@ void FLidarPointCloudOctreeNode::InsertPoints(FLidarPointCloudOctree* Tree, cons
 			}
 			else if(bStoreInBucket)
 			{
-				PointBuckets[InGridData.ChildNodeLocation].Add(Points[Index]);
+				const FLidarPointCloudPoint& Other = Points[Index];
+				PointBuckets[InGridData.ChildNodeLocation].Emplace(AdjustedLocation, Other.Color, !!Other.bVisible, Other.ClassificationID);
 			}
 		}
 		else
@@ -225,11 +217,38 @@ void FLidarPointCloudOctreeNode::InsertPoints(FLidarPointCloudOctree* Tree, cons
 	{
 		FScopeLock Lock(&MapLock);
 
+		// Make a copy of the data
+		TArray<FLidarPointCloudPoint> AllocatedPoints;
+		BulkData.CopyToArray(AllocatedPoints);
+		bCanReleaseData = false;
+
 		// Rebuild Current Grid Mapping
 		for (int32 i = 0; i < AllocatedPoints.Num(); ++i)
 		{
-			auto InGridData = CalculateGridCellData(&AllocatedPoints[i], Center, LODData);
-			CurrentGridAllocationMap.Add(InGridData.Index, FGridAllocation(i, InGridData));
+			auto InGridData = CalculateGridCellData(AllocatedPoints[i].Location, Center, LODData);
+			FGridAllocation* GridCell = CurrentGridAllocationMap.Find(InGridData.Index);
+
+			// Attempt to allocate the point to this node
+			if (GridCell)
+			{
+				if (InGridData.DistanceFromCenter < GridCell->DistanceFromCenter)
+				{
+					PointBuckets[GridCell->ChildNodeLocation].Add(AllocatedPoints[GridCell->Index]);
+					AllocatedPoints[GridCell->Index] = AllocatedPoints[i];
+					GridCell->DistanceFromCenter = InGridData.DistanceFromCenter;
+				}
+				else
+				{
+					PointBuckets[InGridData.ChildNodeLocation].Add(AllocatedPoints[i]);
+				}
+
+				AllocatedPoints.RemoveAtSwap(i--, 1, false);
+				--NumPointsAdded;
+			}
+			else
+			{
+				CurrentGridAllocationMap.Add(InGridData.Index, FGridAllocation(i, InGridData));
+			}
 		}
 
 		// Compare the incoming data to the currently held set, and replace if necessary
@@ -238,15 +257,17 @@ void FLidarPointCloudOctreeNode::InsertPoints(FLidarPointCloudOctree* Tree, cons
 			const int32& GridIndex = Element.Key;
 			const FLidarPointCloudPoint& Point = Points[Element.Value.Index];
 			FGridAllocation* GridCell = CurrentGridAllocationMap.Find(GridIndex);
+			const FVector AdjustedLocation = Point.Location + Translation;
 
 			// Attempt to allocate the point to this node
 			if (GridCell)
 			{
+				FLidarPointCloudPoint& AllocatedPoint = AllocatedPoints[GridCell->Index];
 				bool bStoreInBucket = true;
 
-				if (DuplicateHandling != ELidarPointCloudDuplicateHandling::Ignore && AllocatedPoints[GridCell->Index].Location.Equals(Point.Location, MaxDistanceForDuplicate))
+				if (DuplicateHandling != ELidarPointCloudDuplicateHandling::Ignore && AllocatedPoint.Location.Equals(AdjustedLocation, MaxDistanceForDuplicate))
 				{
-					if (DuplicateHandling == ELidarPointCloudDuplicateHandling::SelectFirst || BrightnessFromColor(Point.Color) <= BrightnessFromColor(AllocatedPoints[GridCell->Index].Color))
+					if (DuplicateHandling == ELidarPointCloudDuplicateHandling::SelectFirst || BrightnessFromColor(Point.Color) <= BrightnessFromColor(AllocatedPoint.Color))
 					{
 						continue;
 					}
@@ -261,33 +282,28 @@ void FLidarPointCloudOctreeNode::InsertPoints(FLidarPointCloudOctree* Tree, cons
 				{
 					if (bStoreInBucket)
 					{
-						PointBuckets[GridCell->ChildNodeLocation].Add(AllocatedPoints[GridCell->Index]);
+						PointBuckets[GridCell->ChildNodeLocation].Add(AllocatedPoint);
 					}
 
-					AllocatedPoints[GridCell->Index] = Point;
+					AllocatedPoint.Location = AdjustedLocation;
+					AllocatedPoint.Color = Point.Color;
+					AllocatedPoint.bVisible = Point.bVisible;
+					AllocatedPoint.ClassificationID = Point.ClassificationID;
 					GridCell->DistanceFromCenter = Element.Value.DistanceFromCenter;
 				}
 				// ... otherwise add it straight to the bucket
 				else if (bStoreInBucket)
 				{
-					PointBuckets[Element.Value.ChildNodeLocation].Add(Point);
+					PointBuckets[Element.Value.ChildNodeLocation].Emplace(AdjustedLocation, Point.Color, !!Point.bVisible, Point.ClassificationID);
 				}
 			}
 			else
 			{
-				CurrentGridAllocationMap.Add(GridIndex, FGridAllocation(AllocatedPoints.Add(Point), Element.Value));
+				CurrentGridAllocationMap.Add(GridIndex, FGridAllocation(AllocatedPoints.Emplace(AdjustedLocation, Point.Color, !!Point.bVisible, Point.ClassificationID), Element.Value));
 				NumPointsAdded++;
 			}
 		}
 
-		// Distribute PaddingPoints into the relevant arrays of PointBuckets
-		for (auto& Point : PaddingPoints)
-		{
-			PointBuckets[CalculateGridCellData(&Point, Center, LODData).ChildNodeLocation].Add(Point);
-		}
-		NumPointsAdded -= PaddingPoints.Num();
-		PaddingPoints.Reset();
-
 		for (uint8 i = 0; i < 8; ++i)
 		{
 			if (!GetChildNodeAtLocation(i))
@@ -296,26 +312,27 @@ void FLidarPointCloudOctreeNode::InsertPoints(FLidarPointCloudOctree* Tree, cons
 				if (Depth < FLidarPointCloudOctree::MaxNodeDepth && PointBuckets[i].Num() > FLidarPointCloudOctree::MaxBucketSize)
 				{
 					const FVector ChildNodeCenter = Center + LODData.Extent * (FVector(-0.5f) + FVector((i & 4) == 4, (i & 2) == 2, (i & 1) == 1));
-					Children.Add(new FLidarPointCloudOctreeNode(Tree, this, Depth + 1, i, ChildNodeCenter));
+					Children.Add(new FLidarPointCloudOctreeNode(Tree,  Depth + 1, i, ChildNodeCenter));
 
 					// The recursive InserPoints call will happen later, after the Lock is released
 				}
 				// ... otherwise, points can be re-added back as padding
 				else
 				{
-					PaddingPoints.Append(PointBuckets[i]);
+					AllocatedPoints.Append(PointBuckets[i]);
+					NumPointsAdded += PointBuckets[i].Num();
 					PointBuckets[i].Reset();
 				}
 			}
 		}
 
-		NumPointsAdded += PaddingPoints.Num();
-
 		// Shrink the data usage
 		AllocatedPoints.Shrink();
-		PaddingPoints.Shrink();
+
+		// Update the BulkData with the new array
+		BulkData.CopyFromArray(AllocatedPoints);
 	}
-	
+
 	AddPointCount(Tree, NumPointsAdded);
 
 	// Pass surplus points
@@ -323,12 +340,12 @@ void FLidarPointCloudOctreeNode::InsertPoints(FLidarPointCloudOctree* Tree, cons
 	{
 		if (PointBuckets[i].Num() > 0)
 		{
-			GetChildNodeAtLocation(i)->InsertPoints(Tree, PointBuckets[i].GetData(), PointBuckets[i].Num(), DuplicateHandling);
+			GetChildNodeAtLocation(i)->InsertPoints(Tree, PointBuckets[i].GetData(), PointBuckets[i].Num(), DuplicateHandling, FVector::ZeroVector);
 		}
 	}
 }
 
-void FLidarPointCloudOctreeNode::InsertPoints(FLidarPointCloudOctree* Tree, FLidarPointCloudPoint** Points, const int32& Count, ELidarPointCloudDuplicateHandling DuplicateHandling)
+void FLidarPointCloudOctreeNode::InsertPoints(FLidarPointCloudOctree* Tree, FLidarPointCloudPoint** Points, const int32& Count, ELidarPointCloudDuplicateHandling DuplicateHandling, const FVector& Translation)
 {
 	const auto& LODData = Tree->SharedData[Depth];
 
@@ -343,7 +360,8 @@ void FLidarPointCloudOctreeNode::InsertPoints(FLidarPointCloudOctree* Tree, FLid
 	// Filter the local set of incoming data
 	for (int32 Index = 0; Index < Count; Index++)
 	{
-		auto InGridData = CalculateGridCellData(Points[Index], Center, LODData);
+		const FVector AdjustedLocation = Points[Index]->Location + Translation;
+		auto InGridData = CalculateGridCellData(AdjustedLocation, Center, LODData);
 		FGridAllocation* GridCell = NewGridAllocationMap.Find(InGridData.Index);
 
 		// Attempt to allocate the point to this node
@@ -367,7 +385,8 @@ void FLidarPointCloudOctreeNode::InsertPoints(FLidarPointCloudOctree* Tree, FLid
 			{
 				if (bStoreInBucket)
 				{
-					PointBuckets[GridCell->ChildNodeLocation].Emplace(*Points[GridCell->Index]);
+					const FLidarPointCloudPoint& Other = *Points[GridCell->Index];
+					PointBuckets[GridCell->ChildNodeLocation].Emplace(AdjustedLocation, Other.Color, !!Other.bVisible, Other.ClassificationID);
 				}
 
 				GridCell->Index = Index;
@@ -375,7 +394,8 @@ void FLidarPointCloudOctreeNode::InsertPoints(FLidarPointCloudOctree* Tree, FLid
 			}
 			else if (bStoreInBucket)
 			{
-				PointBuckets[InGridData.ChildNodeLocation].Emplace(*Points[Index]);
+				const FLidarPointCloudPoint& Other = *Points[GridCell->Index];
+				PointBuckets[InGridData.ChildNodeLocation].Emplace(AdjustedLocation, Other.Color, !!Other.bVisible, Other.ClassificationID);
 			}
 		}
 		else
@@ -388,29 +408,57 @@ void FLidarPointCloudOctreeNode::InsertPoints(FLidarPointCloudOctree* Tree, FLid
 	{
 		FScopeLock Lock(&MapLock);
 
+		// Make a copy of the data
+		TArray<FLidarPointCloudPoint> AllocatedPoints;
+		BulkData.CopyToArray(AllocatedPoints);
+		bCanReleaseData = false;
+
 		// Rebuild Current Grid Mapping
-		for (int32 i = 0; i < AllocatedPoints.Num(); i++)
+		for (int32 i = 0; i < AllocatedPoints.Num(); ++i)
 		{
-			auto InGridData = CalculateGridCellData(&AllocatedPoints[i], Center, LODData);
-			CurrentGridAllocationMap.Add(InGridData.Index, FGridAllocation(i, InGridData));
+			auto InGridData = CalculateGridCellData(AllocatedPoints[i].Location, Center, LODData);
+			FGridAllocation* GridCell = CurrentGridAllocationMap.Find(InGridData.Index);
+
+			// Attempt to allocate the point to this node
+			if (GridCell)
+			{
+				if (InGridData.DistanceFromCenter < GridCell->DistanceFromCenter)
+				{
+					PointBuckets[GridCell->ChildNodeLocation].Add(AllocatedPoints[GridCell->Index]);
+					AllocatedPoints[GridCell->Index] = AllocatedPoints[i];
+					GridCell->DistanceFromCenter = InGridData.DistanceFromCenter;
+				}
+				else
+				{
+					PointBuckets[InGridData.ChildNodeLocation].Add(AllocatedPoints[i]);
+				}
+
+				AllocatedPoints.RemoveAtSwap(i--, 1, false);
+				--NumPointsAdded;
+			}
+			else
+			{
+				CurrentGridAllocationMap.Add(InGridData.Index, FGridAllocation(i, InGridData));
+			}
 		}
 
 		// Compare the incoming data to the currently held set, and replace if necessary
-		FLidarPointCloudPoint Point;
 		for (auto& Element : NewGridAllocationMap)
 		{
 			const int32& GridIndex = Element.Key;
-			Point.CopyFrom(*Points[Element.Value.Index]);
+			const FLidarPointCloudPoint& Point = *Points[Element.Value.Index];
 			FGridAllocation* GridCell = CurrentGridAllocationMap.Find(GridIndex);
+			FVector AdjustedLocation = Point.Location + Translation;
 
 			// Attempt to allocate the point to this node
 			if (GridCell)
 			{
+				FLidarPointCloudPoint& AllocatedPoint = AllocatedPoints[GridCell->Index];
 				bool bStoreInBucket = true;
 
-				if (DuplicateHandling != ELidarPointCloudDuplicateHandling::Ignore && AllocatedPoints[GridCell->Index].Location.Equals(Point.Location, MaxDistanceForDuplicate))
+				if (DuplicateHandling != ELidarPointCloudDuplicateHandling::Ignore && AllocatedPoint.Location.Equals(AdjustedLocation, MaxDistanceForDuplicate))
 				{
-					if (DuplicateHandling == ELidarPointCloudDuplicateHandling::SelectFirst || BrightnessFromColor(Point.Color) <= BrightnessFromColor(AllocatedPoints[GridCell->Index].Color))
+					if (DuplicateHandling == ELidarPointCloudDuplicateHandling::SelectFirst || BrightnessFromColor(Point.Color) <= BrightnessFromColor(AllocatedPoint.Color))
 					{
 						continue;
 					}
@@ -425,33 +473,28 @@ void FLidarPointCloudOctreeNode::InsertPoints(FLidarPointCloudOctree* Tree, FLid
 				{
 					if (bStoreInBucket)
 					{
-						PointBuckets[GridCell->ChildNodeLocation].Add(AllocatedPoints[GridCell->Index]);
+						PointBuckets[GridCell->ChildNodeLocation].Add(AllocatedPoint);
 					}
 
-					AllocatedPoints[GridCell->Index] = Point;
+					AllocatedPoint.Location = AdjustedLocation;
+					AllocatedPoint.Color = Point.Color;
+					AllocatedPoint.bVisible = Point.bVisible;
+					AllocatedPoint.ClassificationID = Point.ClassificationID;
 					GridCell->DistanceFromCenter = Element.Value.DistanceFromCenter;
 				}
 				// ... otherwise add it straight to the bucket
 				else if (bStoreInBucket)
 				{
-					PointBuckets[Element.Value.ChildNodeLocation].Add(Point);
+					PointBuckets[Element.Value.ChildNodeLocation].Emplace(AdjustedLocation, Point.Color, !!Point.bVisible, Point.ClassificationID);
 				}
 			}
 			else
 			{
-				CurrentGridAllocationMap.Add(GridIndex, FGridAllocation(AllocatedPoints.Add(Point), Element.Value));
+				CurrentGridAllocationMap.Add(GridIndex, FGridAllocation(AllocatedPoints.Emplace(AdjustedLocation, Point.Color, !!Point.bVisible, Point.ClassificationID), Element.Value));
 				NumPointsAdded++;
 			}
 		}
 
-		// Distribute PaddingPoints into the relevant arrays of PointBuckets
-		for (auto& PaddingPoint : PaddingPoints)
-		{
-			PointBuckets[CalculateGridCellData(&PaddingPoint, Center, LODData).ChildNodeLocation].Add(PaddingPoint);
-		}
-		NumPointsAdded -= PaddingPoints.Num();
-		PaddingPoints.Reset();
-
 		for (uint8 i = 0; i < 8; i++)
 		{
 			if (!GetChildNodeAtLocation(i))
@@ -460,24 +503,25 @@ void FLidarPointCloudOctreeNode::InsertPoints(FLidarPointCloudOctree* Tree, FLid
 				if (Depth < FLidarPointCloudOctree::MaxNodeDepth && PointBuckets[i].Num() > FLidarPointCloudOctree::MaxBucketSize)
 				{
 					const FVector ChildNodeCenter = Center + LODData.Extent * (FVector(-0.5f) + FVector((i & 4) == 4, (i & 2) == 2, (i & 1) == 1));
-					Children.Add(new FLidarPointCloudOctreeNode(Tree, this, Depth + 1, i, ChildNodeCenter));
+					Children.Add(new FLidarPointCloudOctreeNode(Tree,  Depth + 1, i, ChildNodeCenter));
 
 					// The recursive InserPoints call will happen later, after the Lock is released
 				}
 				// ... otherwise, points can be re-added back as padding
 				else
 				{
-					PaddingPoints.Append(PointBuckets[i]);
+					AllocatedPoints.Append(PointBuckets[i]);
+					NumPointsAdded += PointBuckets[i].Num();
 					PointBuckets[i].Reset();
 				}
 			}
 		}
 
-		NumPointsAdded += PaddingPoints.Num();
-
 		// Shrink the data usage
 		AllocatedPoints.Shrink();
-		PaddingPoints.Shrink();
+
+		// Update the BulkData with the new array
+		BulkData.CopyFromArray(AllocatedPoints);
 	}
 
 	AddPointCount(Tree, NumPointsAdded);
@@ -487,15 +531,14 @@ void FLidarPointCloudOctreeNode::InsertPoints(FLidarPointCloudOctree* Tree, FLid
 	{
 		if (PointBuckets[i].Num() > 0)
 		{
-			GetChildNodeAtLocation(i)->InsertPoints(Tree, PointBuckets[i].GetData(), PointBuckets[i].Num(), DuplicateHandling);
+			GetChildNodeAtLocation(i)->InsertPoints(Tree, PointBuckets[i].GetData(), PointBuckets[i].Num(), DuplicateHandling, FVector::ZeroVector);
 		}
 	}
 }
 
 void FLidarPointCloudOctreeNode::Empty(bool bRecursive)
 {
-	AllocatedPoints.Empty();
-	PaddingPoints.Empty();
+	BulkData.RemoveBulkData();
 
 	if (bRecursive)
 	{
@@ -518,53 +561,38 @@ uint32 FLidarPointCloudOctreeNode::GetMaxDepth() const
 	return MaxDepth;
 }
 
-int64 FLidarPointCloudOctreeNode::GetAllocatedSize(bool bRecursive) const
+int64 FLidarPointCloudOctreeNode::GetAllocatedSize(bool bRecursive, bool bIncludeBulkData) const
 {
 	int64 Size = sizeof(FLidarPointCloudOctreeNode);
 
 	Size += Children.GetAllocatedSize();
-	Size += AllocatedPoints.GetAllocatedSize();
-	Size += PaddingPoints.GetAllocatedSize();
+
+	if (bIncludeBulkData)
+	{
+		Size += BulkData.GetBulkDataSize();
+	}
 
 	if (bRecursive)
 	{
 		for (FLidarPointCloudOctreeNode* Child : Children)
 		{
-			Size += Child->GetAllocatedSize(true);
+			Size += Child->GetAllocatedSize(true, bIncludeBulkData);
 		}
 	}
 
 	return Size;
 }
 
-void FLidarPointCloudOctreeNode::Serialize(FArchive& Ar, FLidarPointCloudOctree* Tree)
+void FLidarPointCloudOctreeNode::ReleaseData(bool bForce)
 {
-	// This is needed as LocationInParent doesn't use a full byte and the FArchive complains
-	uint8 FullLocationInParent = LocationInParent;
-	Ar << FullLocationInParent << Center << AllocatedPoints << PaddingPoints;
-
-	int32 NumChildren = Children.Num();
-	Ar << NumChildren;
-
-	if (Ar.IsLoading())
+	// Ignore request, if the node cannot be released
+	if (!bCanReleaseData && !bForce)
 	{
-		LocationInParent = FullLocationInParent;
-
-		AddPointCount(Tree);
-
-		bVisibilityDirty = true;
-
-		Children.AddUninitialized(NumChildren);
-		for (int32 i = 0; i < NumChildren; i++)
-		{
-			Children[i] = new FLidarPointCloudOctreeNode(Tree, this, Depth + 1);
-		}
+		return;
 	}
 
-	for (int32 i = 0; i < NumChildren; i++)
-	{
-		Children[i]->Serialize(Ar, Tree);
-	}
+	bHasDataPending = false;
+	BulkData.ReleaseData();
 }
 
 void FLidarPointCloudOctreeNode::AddPointCount(FLidarPointCloudOctree* Tree, int32 PointCount /*= INT32_MIN*/)
@@ -575,13 +603,29 @@ void FLidarPointCloudOctreeNode::AddPointCount(FLidarPointCloudOctree* Tree, int
 	NumVisiblePoints += Count;
 }
 
+void FLidarPointCloudOctreeNode::SortVisiblePoints()
+{
+	TArrayView<FLidarPointCloudPoint> Points(GetData(), GetNumPoints());
+	Algo::Sort(Points, [](const FLidarPointCloudPoint& A, const FLidarPointCloudPoint& B)
+	{
+		return A.bVisible > B.bVisible;
+	});
+	bCanReleaseData = false;
+}
+
 //////////////////////////////////////////////////////////// FLidarPointCloudOctree
 
-FLidarPointCloudOctree::FLidarPointCloudOctree()
+FLidarPointCloudOctree::FLidarPointCloudOctree(ULidarPointCloud* Owner)
+	: Owner(Owner)
+	, bStreamingBusy(false)
+	, bIsFullyLoaded(false)
 {
 	PointCount.AddDefaulted(MaxNodeDepth + 1);
 	NodeCount.AddDefaulted(MaxNodeDepth + 1);
 	SharedData.AddDefaulted(MaxNodeDepth + 1);
+
+	// Account for the Root
+	NodeCount[0].Increment();
 }
 
 FLidarPointCloudOctree::~FLidarPointCloudOctree()
@@ -604,29 +648,31 @@ int32 FLidarPointCloudOctree::GetNumLODs() const
 	return LODCount;
 }
 
-void FLidarPointCloudOctree::RefreshPointsBounds()
+void FLidarPointCloudOctree::RefreshBounds()
 {
-	PointsBounds.Init();
+	FBox Bounds(EForceInit::ForceInit);
 
-	TQueue<const FLidarPointCloudOctreeNode*> Nodes;
-	const FLidarPointCloudOctreeNode* CurrentNode = nullptr;
-	Nodes.Enqueue(&Root);
-	while (Nodes.Dequeue(CurrentNode))
+	// Calculate the current bounds
+	ITERATE_NODES_CONST({ FOR_RO(Point, CurrentNode) { Bounds += Point->Location; } }, true);
+
+	Extent = Bounds.GetExtent();
+	FVector Offset = Bounds.GetCenter();
+
+	if (!Offset.IsNearlyZero(0.1f))
 	{
-		FOR(Point, CurrentNode->AllocatedPoints)
-		{
-			PointsBounds += Point->Location;
-		}
+		Owner->LocationOffset += Offset;
+		Owner->OriginalCoordinates += Offset;
 
-		FOR(Point, CurrentNode->PaddingPoints)
+		// Shift the points back to the relative position
+		ITERATE_NODES(
 		{
-			PointsBounds += Point->Location;
-		}
+			CurrentNode->Center -= Offset;
 
-		for (auto Child : CurrentNode->Children)
-		{
-			Nodes.Enqueue(Child);
-		}
+			FOR(Point, CurrentNode)
+			{
+				Point->Location -= Offset;
+			}
+		}, true);
 	}
 }
 
@@ -673,29 +719,24 @@ int32 FLidarPointCloudOctree::GetNumNodes() const
 
 int64 FLidarPointCloudOctree::GetAllocatedSize() const
 {
-	// Use cached version, if possible
-	if (PreviousPointCount == GetNumPoints() && PreviousNodeCount == GetNumNodes())
+	if (PreviousPointCount != GetNumPoints() || PreviousNodeCount != GetNumNodes())
 	{
-		return PreviousAllocatedSize;
+		FLidarPointCloudOctree* mutable_this = const_cast<FLidarPointCloudOctree*>(this);
+		mutable_this->RefreshAllocatedSize();
 	}
 
-	FScopeTryLock Lock(&DataLock);
-	if (!Lock.IsLocked())
+	return PreviousAllocatedSize;
+}
+
+int64 FLidarPointCloudOctree::GetAllocatedStructureSize() const
+{
+	if (PreviousPointCount != GetNumPoints() || PreviousNodeCount != GetNumNodes())
 	{
-		return 0;
+		FLidarPointCloudOctree* mutable_this = const_cast<FLidarPointCloudOctree*>(this);
+		mutable_this->RefreshAllocatedSize();
 	}
 
-	PreviousPointCount = GetNumPoints();
-	PreviousNodeCount = GetNumNodes();
-
-	PreviousAllocatedSize = sizeof(FLidarPointCloudOctree);
-
-	PreviousAllocatedSize += SharedData.GetAllocatedSize();
-	PreviousAllocatedSize += PointCount.GetAllocatedSize();
-
-	PreviousAllocatedSize += Root.GetAllocatedSize(true);
-
-	return PreviousAllocatedSize;
+	return PreviousAllocatedStructureSize;
 }
 
 float FLidarPointCloudOctree::GetEstimatedPointSpacing() const
@@ -724,7 +765,7 @@ void FLidarPointCloudOctree::RemoveCollision()
 	new (&CollisionMesh) FTriMeshCollisionData();
 }
 
-void FLidarPointCloudOctree::GetPoints(TArray<FLidarPointCloudPoint*>& Points, int64 StartIndex, int64 Count) const
+void FLidarPointCloudOctree::GetPoints(TArray<FLidarPointCloudPoint*>& Points, int64 StartIndex, int64 Count)
 {
 	check(StartIndex >= 0 && StartIndex < GetNumPoints());
 
@@ -742,7 +783,7 @@ void FLidarPointCloudOctree::GetPoints(TArray<FLidarPointCloudPoint*>& Points, i
 
 	FScopeLock Lock(&DataLock);
 
-	ITERATE_NODES_CONST(
+	ITERATE_NODES(
 	{
 		// If no data is required, quit
 		if (Count == 0)
@@ -762,13 +803,9 @@ void FLidarPointCloudOctree::GetPoints(TArray<FLidarPointCloudPoint*>& Points, i
 				FLidarPointCloudPoint** Dest = Points.GetData() + Offset;
 
 				// Assign pointers to the data
-				for (auto& Point : CurrentNode->AllocatedPoints)
-				{
-					*Dest++ = &Point;
-				}
-				for (auto& Point : CurrentNode->PaddingPoints)
+				FOR(Point, CurrentNode)
 				{
-					*Dest++ = &Point;
+					*Dest++ = Point;
 				}
 
 				Count -= CurrentNode->GetNumPoints();
@@ -776,52 +813,25 @@ void FLidarPointCloudOctree::GetPoints(TArray<FLidarPointCloudPoint*>& Points, i
 			// ... otherwise iterate over needed data
 			else
 			{
-				// AllocatedPoints
+				int32 NumPointsToCopy = FMath::Max(0LL, FMath::Min((int64)CurrentNode->GetNumPoints(), Count + StartIndex) - StartIndex);
+				if (NumPointsToCopy > 0)
 				{
-					int32 NumPointsToCopy = FMath::Max(0LL, FMath::Min((int64)CurrentNode->AllocatedPoints.Num(), Count + StartIndex) - StartIndex);
-					if (NumPointsToCopy > 0)
-					{
-						// Expand the array
-						int32 Offset = Points.Num();
-						Points.AddUninitialized(NumPointsToCopy);
-												
-						// Setup pointers
-						FLidarPointCloudPoint* Src = CurrentNode->AllocatedPoints.GetData() + StartIndex;
-						FLidarPointCloudPoint** Dest = Points.GetData() + Offset;
-
-						// Assign source pointers to the destination
-						for (int32 i = 0; i < NumPointsToCopy; i++)
-						{
-							*Dest++ = Src++;
-						}
+					// Expand the array
+					int32 Offset = Points.Num();
+					Points.AddUninitialized(NumPointsToCopy);
 
-						StartIndex = FMath::Max(0LL, StartIndex - NumPointsToCopy);
-						Count -= NumPointsToCopy;
-					}
-				}
+					// Setup pointers
+					FLidarPointCloudPoint* Src = CurrentNode->GetData() + StartIndex;
+					FLidarPointCloudPoint** Dest = Points.GetData() + Offset;
 
-				// PaddingPoints
-				{
-					int32 NumPointsToCopy = FMath::Max(0LL, FMath::Min((int64)CurrentNode->PaddingPoints.Num(), Count + StartIndex) - StartIndex);
-					if (NumPointsToCopy > 0)
+					// Assign source pointers to the destination
+					for (int32 i = 0; i < NumPointsToCopy; i++)
 					{
-						// Expand the array
-						int32 Offset = Points.Num();
-						Points.AddUninitialized(NumPointsToCopy);
-
-						// Setup pointers
-						FLidarPointCloudPoint* Src = CurrentNode->PaddingPoints.GetData() + StartIndex;
-						FLidarPointCloudPoint** Dest = Points.GetData() + Offset;
-
-						// Assign source pointers to the destination
-						for (int32 i = 0; i < NumPointsToCopy; i++)
-						{
-							*Dest++ = Src++;
-						}
-
-						StartIndex = FMath::Max(0LL, StartIndex - NumPointsToCopy);;
-						Count -= NumPointsToCopy;
+						*Dest++ = Src++;
 					}
+
+					StartIndex = FMath::Max(0LL, StartIndex - NumPointsToCopy);
+					Count -= NumPointsToCopy;
 				}
 			}
 		}
@@ -833,25 +843,25 @@ void FLidarPointCloudOctree::GetPoints(TArray<FLidarPointCloudPoint*>& Points, i
 	}, true);
 }
 
-void FLidarPointCloudOctree::GetPointsInSphere(TArray<FLidarPointCloudPoint*>& SelectedPoints, const FSphere& Sphere, const bool& bVisibleOnly) const
+void FLidarPointCloudOctree::GetPointsInSphere(TArray<FLidarPointCloudPoint*>& SelectedPoints, const FSphere& Sphere, const bool& bVisibleOnly)
 {
 	SelectedPoints.Reset();
-	PROCESS_IN_SPHERE_CONST({ SelectedPoints.Add(Point); });
+	PROCESS_IN_SPHERE({ SelectedPoints.Add(Point); });
 }
 
-void FLidarPointCloudOctree::GetPointsInBox(TArray<FLidarPointCloudPoint*>& SelectedPoints, const FBox& Box, const bool& bVisibleOnly) const
+void FLidarPointCloudOctree::GetPointsInBox(TArray<FLidarPointCloudPoint*>& SelectedPoints, const FBox& Box, const bool& bVisibleOnly)
 {
 	SelectedPoints.Reset(); 
-	PROCESS_IN_BOX_CONST({ SelectedPoints.Add(Point); });
+	PROCESS_IN_BOX({ SelectedPoints.Add(Point); });
 }
 
-void FLidarPointCloudOctree::GetPointsInFrustum(TArray<FLidarPointCloudPoint*>& SelectedPoints, const FConvexVolume& Frustum, const bool& bVisibleOnly) const
+void FLidarPointCloudOctree::GetPointsInFrustum(TArray<FLidarPointCloudPoint*>& SelectedPoints, const FConvexVolume& Frustum, const bool& bVisibleOnly)
 {
 	SelectedPoints.Reset();
-	PROCESS_IN_FRUSTUM_CONST({ SelectedPoints.Add(Point); });
+	PROCESS_IN_FRUSTUM({ SelectedPoints.Add(Point); });
 }
 
-void FLidarPointCloudOctree::GetPointsAsCopies(TArray<FLidarPointCloudPoint>& Points, int64 StartIndex, int64 Count) const
+void FLidarPointCloudOctree::GetPointsAsCopies(TArray<FLidarPointCloudPoint>& Points, const FTransform* LocalToWorld, int64 StartIndex, int64 Count) const
 {
 	// If empty, abort
 	if (GetNumPoints() == 0)
@@ -876,53 +886,59 @@ void FLidarPointCloudOctree::GetPointsAsCopies(TArray<FLidarPointCloudPoint>& Po
 
 	FScopeLock Lock(&DataLock);
 
-	TQueue<const FLidarPointCloudOctreeNode*> Nodes;
-	const FLidarPointCloudOctreeNode* CurrentNode = nullptr;
-	Nodes.Enqueue(&Root);
-	while (Nodes.Dequeue(CurrentNode))
+	if (LocalToWorld)
 	{
-		// If no data is required, quit
-		if (Count == 0)
-		{
-			return;
-		}
-
-		// Should this node's points be injected?
-		if (StartIndex < CurrentNode->GetNumPoints())
-		{
-			// If the index is at 0 and the requested count is at least the size of this node, append whole arrays
-			if (StartIndex == 0 && Count >= CurrentNode->GetNumPoints())
+		ITERATE_NODES_CONST({
+			// If no data is required, quit
+			if (Count == 0)
 			{
-				Points.Append(CurrentNode->AllocatedPoints);
-				Points.Append(CurrentNode->PaddingPoints);
-
-				Count -= CurrentNode->GetNumPoints();
+				return;
 			}
-			// ... otherwise iterate over needed data
-			else
+
+			// Should this node's points be injected?
+			if (StartIndex < CurrentNode->GetNumPoints())
 			{
-				// AllocatedPoints
+				const int32 NumPointsToCopy = FMath::Max(0LL, FMath::Min((int64)CurrentNode->GetNumPoints(), Count + StartIndex) - StartIndex);
+				if (NumPointsToCopy > 0)
 				{
-					int32 NumPointsToCopy = FMath::Max(0LL, FMath::Min((int64)CurrentNode->AllocatedPoints.Num(), Count + StartIndex) - StartIndex);
-					if (NumPointsToCopy > 0)
+					for (auto Point = CurrentNode->GetData() + StartIndex, DataEnd = Point + StartIndex + NumPointsToCopy; Point != DataEnd; ++Point)
 					{
-						int32 SrcOffset = StartIndex;
-						int32 DestOffset = Points.Num();
-
-						// Expand the array
-						Points.AddUninitialized(NumPointsToCopy);
-
-						// Copy contents
-						FMemory::Memcpy(Points.GetData() + DestOffset, CurrentNode->AllocatedPoints.GetData() + SrcOffset, NumPointsToCopy * sizeof(FLidarPointCloudPoint));
-
-						StartIndex = FMath::Max(0LL, StartIndex - NumPointsToCopy);
-						Count -= NumPointsToCopy;
+						Points.Add(Point->Transform(*LocalToWorld));
 					}
+
+					StartIndex = FMath::Max(0LL, StartIndex - NumPointsToCopy);
+					Count -= NumPointsToCopy;
 				}
+			}
+			// ... or skipped completely?
+			else
+			{
+				StartIndex -= CurrentNode->GetNumPoints();
+			}
+		}, true);
+	}
+	else
+	{
+		ITERATE_NODES_CONST({
+			// If no data is required, quit
+			if (Count == 0)
+			{
+				return;
+			}
 
-				// PaddingPoints
+			// Should this node's points be injected?
+			if (StartIndex < CurrentNode->GetNumPoints())
+			{
+				// If the index is at 0 and the requested count is at least the size of this node, append whole arrays
+				if (StartIndex == 0 && Count >= CurrentNode->GetNumPoints())
+				{
+					Points.Append(CurrentNode->GetData(), CurrentNode->GetNumPoints());
+					Count -= CurrentNode->GetNumPoints();
+				}
+				// ... otherwise iterate over needed data
+				else
 				{
-					int32 NumPointsToCopy = FMath::Max(0LL, FMath::Min((int64)CurrentNode->PaddingPoints.Num(), Count + StartIndex) - StartIndex);
+					int32 NumPointsToCopy = FMath::Max(0LL, FMath::Min((int64)CurrentNode->GetNumPoints(), Count + StartIndex) - StartIndex);
 					if (NumPointsToCopy > 0)
 					{
 						int32 SrcOffset = StartIndex;
@@ -932,56 +948,72 @@ void FLidarPointCloudOctree::GetPointsAsCopies(TArray<FLidarPointCloudPoint>& Po
 						Points.AddUninitialized(NumPointsToCopy);
 
 						// Copy contents
-						FMemory::Memcpy(Points.GetData() + DestOffset, CurrentNode->PaddingPoints.GetData() + SrcOffset, NumPointsToCopy * sizeof(FLidarPointCloudPoint));
+						FMemory::Memcpy(Points.GetData() + DestOffset, CurrentNode->GetData() + SrcOffset, NumPointsToCopy * sizeof(FLidarPointCloudPoint));
 
-						StartIndex = FMath::Max(0LL, StartIndex - NumPointsToCopy);;
+						StartIndex = FMath::Max(0LL, StartIndex - NumPointsToCopy);
 						Count -= NumPointsToCopy;
 					}
 				}
 			}
-		}
-		// ... or skipped completely?
-		else
-		{
-			StartIndex -= CurrentNode->GetNumPoints();
-		}
-
-		for (auto Child : CurrentNode->Children)
-		{
-			Nodes.Enqueue(Child);
-		}
+			// ... or skipped completely?
+			else
+			{
+				StartIndex -= CurrentNode->GetNumPoints();
+			}
+		}, true);
 	}
 }
 
-void FLidarPointCloudOctree::GetPointsInSphereAsCopies(TArray<FLidarPointCloudPoint>& SelectedPoints, const FSphere& Sphere, const bool& bVisibleOnly) const
+void FLidarPointCloudOctree::GetPointsInSphereAsCopies(TArray<FLidarPointCloudPoint>& SelectedPoints, const FSphere& Sphere, const bool& bVisibleOnly, const FTransform* LocalToWorld) const
 {
 	SelectedPoints.Reset();
-	PROCESS_IN_SPHERE_CONST({ SelectedPoints.Add(*Point); });
+	if (LocalToWorld)
+	{
+		PROCESS_IN_SPHERE_CONST({ SelectedPoints.Add(Point->Transform(*LocalToWorld)); });
+	}
+	else
+	{
+		PROCESS_IN_SPHERE_CONST({ SelectedPoints.Add(*Point); });
+	}
 }
 
-void FLidarPointCloudOctree::GetPointsInBoxAsCopies(TArray<FLidarPointCloudPoint>& SelectedPoints, const FBox& Box, const bool& bVisibleOnly) const
+void FLidarPointCloudOctree::GetPointsInBoxAsCopies(TArray<FLidarPointCloudPoint>& SelectedPoints, const FBox& Box, const bool& bVisibleOnly, const FTransform* LocalToWorld) const
 {
 	SelectedPoints.Reset();
-	PROCESS_IN_BOX_CONST({ SelectedPoints.Add(*Point); });
+	if (LocalToWorld)
+	{
+		PROCESS_IN_BOX_CONST({ SelectedPoints.Add(Point->Transform(*LocalToWorld)); });
+	}
+	else
+	{
+		PROCESS_IN_BOX_CONST({ SelectedPoints.Add(*Point); });
+	}
 }
 
 FLidarPointCloudPoint* FLidarPointCloudOctree::RaycastSingle(const FLidarPointCloudRay& Ray, const float& Radius, const bool& bVisibleOnly)
 {
-	PROCESS_BY_RAY_CONST({ return Point; });
+	PROCESS_BY_RAY({ return Point; });
 	return nullptr;
 }
 
 bool FLidarPointCloudOctree::RaycastMulti(const FLidarPointCloudRay& Ray, const float& Radius, const bool& bVisibleOnly, TArray<FLidarPointCloudPoint*>& OutHits)
 {
 	OutHits.Reset();
-	PROCESS_BY_RAY_CONST({ OutHits.Add(Point); });
+	PROCESS_BY_RAY({ OutHits.Add(Point); });
 	return OutHits.Num() > 0;
 }
 
-bool FLidarPointCloudOctree::RaycastMulti(const FLidarPointCloudRay& Ray, const float& Radius, const bool& bVisibleOnly, TArray<FLidarPointCloudPoint>& OutHits)
+bool FLidarPointCloudOctree::RaycastMulti(const FLidarPointCloudRay& Ray, const float& Radius, const bool& bVisibleOnly, const FTransform* LocalToWorld, TArray<FLidarPointCloudPoint>& OutHits)
 {
 	OutHits.Reset();
-	PROCESS_BY_RAY_CONST({ OutHits.Add(*Point); });
+	if (LocalToWorld)
+	{
+		PROCESS_BY_RAY_CONST({ OutHits.Add(Point->Transform(*LocalToWorld)); });
+	}
+	else
+	{
+		PROCESS_BY_RAY_CONST({ OutHits.Add(*Point); });
+	}
 	return OutHits.Num() > 0;
 }
 
@@ -1013,7 +1045,10 @@ void FLidarPointCloudOctree::SetVisibilityOfPointsInSphere(const bool& bNewVisib
 			// If node fully inside the radius - do not check individual points
 			if (CurrentNode->GetSphereBounds(this).IsInside(Sphere))
 			{
-				ITERATE_POINTS({ Point->bVisible = bNewVisibility; });
+				FOR(Point, CurrentNode)
+				{
+					Point->bVisible = bNewVisibility;
+				}
 
 				if (bNewVisibility)
 				{
@@ -1022,7 +1057,7 @@ void FLidarPointCloudOctree::SetVisibilityOfPointsInSphere(const bool& bNewVisib
 			}
 			else
 			{
-				ITERATE_POINTS(
+				FOR(Point, CurrentNode)
 				{
 					if (Point->bVisible != bNewVisibility && POINT_IN_SPHERE)
 					{
@@ -1033,26 +1068,17 @@ void FLidarPointCloudOctree::SetVisibilityOfPointsInSphere(const bool& bNewVisib
 					{
 						++CurrentNode->NumVisiblePoints;
 					}
-				});
+				}
 			}
 
 			CurrentNode->bVisibilityDirty = false;
 
 			// Sort points to speed up rendering
-			Algo::Sort(CurrentNode->AllocatedPoints, MoveTemp(PointVisibilitySortingFunction));
-			Algo::Sort(CurrentNode->PaddingPoints, MoveTemp(PointVisibilitySortingFunction));
+			CurrentNode->SortVisiblePoints();
 		}
 	}, NODE_IN_BOX);
 }
 
-void FLidarPointCloudOctree::SetVisibilityOfPointsInSphereAsync(const bool& bNewVisibility, const FSphere& Sphere, TFunction<void(void)> CompletionCallback)
-{
-	Async(EAsyncExecution::Thread, [this, bNewVisibility, Sphere]
-	{
-		SetVisibilityOfPointsInSphere(bNewVisibility, Sphere);
-	}, MoveTemp(CompletionCallback));
-}
-
 void FLidarPointCloudOctree::SetVisibilityOfPointsInBox(const bool& bNewVisibility, const FBox& Box)
 {
 	ITERATE_NODES({
@@ -1065,7 +1091,10 @@ void FLidarPointCloudOctree::SetVisibilityOfPointsInBox(const bool& bNewVisibili
 			// If node fully inside the radius - do not check individual points
 			if (Box.IsInsideOrOn(CurrentNode->Center - SharedData[CurrentNode->Depth].Extent) && Box.IsInsideOrOn(CurrentNode->Center + SharedData[CurrentNode->Depth].Extent))
 			{
-				ITERATE_POINTS({ Point->bVisible = bNewVisibility; });
+				FOR(Point, CurrentNode)
+				{
+					Point->bVisible = bNewVisibility;
+				}
 
 				if (bNewVisibility)
 				{
@@ -1074,7 +1103,7 @@ void FLidarPointCloudOctree::SetVisibilityOfPointsInBox(const bool& bNewVisibili
 			}
 			else
 			{
-				ITERATE_POINTS(
+				FOR(Point, CurrentNode)
 				{
 					if (Point->bVisible != bNewVisibility && POINT_IN_BOX)
 					{
@@ -1085,26 +1114,17 @@ void FLidarPointCloudOctree::SetVisibilityOfPointsInBox(const bool& bNewVisibili
 					{
 						++CurrentNode->NumVisiblePoints;
 					}
-				});
+				}
 			}
 
 			CurrentNode->bVisibilityDirty = false;
 
 			// Sort points to speed up rendering
-			Algo::Sort(CurrentNode->AllocatedPoints, MoveTemp(PointVisibilitySortingFunction));
-			Algo::Sort(CurrentNode->PaddingPoints, MoveTemp(PointVisibilitySortingFunction));
+			CurrentNode->SortVisiblePoints();
 		}
 	}, NODE_IN_BOX);
 }
 
-void FLidarPointCloudOctree::SetVisibilityOfPointsInBoxAsync(const bool& bNewVisibility, const FBox& Box, TFunction<void(void)> CompletionCallback)
-{
-	Async(EAsyncExecution::Thread, [this, bNewVisibility, Box]
-	{
-		SetVisibilityOfPointsInBox(bNewVisibility, Box);
-	}, MoveTemp(CompletionCallback));
-}
-
 void FLidarPointCloudOctree::SetVisibilityOfPointsByRay(const bool& bNewVisibility, const FLidarPointCloudRay& Ray, const float& Radius)
 {
 	const float RadiusSq = Radius * Radius;
@@ -1115,7 +1135,7 @@ void FLidarPointCloudOctree::SetVisibilityOfPointsByRay(const bool& bNewVisibili
 		{
 			CurrentNode->NumVisiblePoints = 0;
 
-			ITERATE_POINTS(
+			FOR(Point, CurrentNode)
 			{
 				if (Point->bVisible != bNewVisibility && POINT_BY_RAY)
 				{
@@ -1126,13 +1146,12 @@ void FLidarPointCloudOctree::SetVisibilityOfPointsByRay(const bool& bNewVisibili
 				{
 					++CurrentNode->NumVisiblePoints;
 				}
-			});
+			}
 
 			CurrentNode->bVisibilityDirty = false;
 
 			// Sort points to speed up rendering
-			Algo::Sort(CurrentNode->AllocatedPoints, MoveTemp(PointVisibilitySortingFunction));
-			Algo::Sort(CurrentNode->PaddingPoints, MoveTemp(PointVisibilitySortingFunction));
+			CurrentNode->SortVisiblePoints();
 
 			for (auto Child : CurrentNode->Children)
 			{
@@ -1142,25 +1161,20 @@ void FLidarPointCloudOctree::SetVisibilityOfPointsByRay(const bool& bNewVisibili
 	}, false);
 }
 
-void FLidarPointCloudOctree::SetVisibilityOfPointsByRayAsync(const bool& bNewVisibility, const FLidarPointCloudRay& Ray, const float& Radius, TFunction<void(void)> CompletionCallback /*= nullptr*/)
-{
-	Async(EAsyncExecution::Thread, [this, bNewVisibility, Ray, Radius]
-	{
-		SetVisibilityOfPointsByRay(bNewVisibility, Ray, Radius);
-	}, MoveTemp(CompletionCallback));
-}
-
 void FLidarPointCloudOctree::HideAll()
 {
 	ITERATE_NODES({
 		if (CurrentNode->NumVisiblePoints > 0)
 		{
-			ITERATE_POINTS({ Point->bVisible = false; });
+			FOR(Point, CurrentNode)
+			{
+				Point->bVisible = false;
+			}
 
 			CurrentNode->NumVisiblePoints = 0;
 			CurrentNode->bVisibilityDirty = false;
 		}
-		}, true);
+	}, true);
 }
 
 void FLidarPointCloudOctree::UnhideAll()
@@ -1168,7 +1182,10 @@ void FLidarPointCloudOctree::UnhideAll()
 	ITERATE_NODES({
 		if (CurrentNode->NumVisiblePoints != CurrentNode->GetNumPoints())
 		{
-			ITERATE_POINTS({ Point->bVisible = true; });
+			FOR(Point, CurrentNode)
+			{
+				Point->bVisible = true;
+			}
 
 			CurrentNode->NumVisiblePoints = CurrentNode->GetNumPoints();
 			CurrentNode->bVisibilityDirty = false;
@@ -1176,66 +1193,26 @@ void FLidarPointCloudOctree::UnhideAll()
 	}, true);
 }
 
-void FLidarPointCloudOctree::ResetVisibilityAsync(TFunction<void(void)> CompletionCallback)
-{
-	Async(EAsyncExecution::Thread, [this]
-	{
-		ResetVisibilityAsync();
-	}, MoveTemp(CompletionCallback));
-}
-
 void FLidarPointCloudOctree::ExecuteActionOnAllPoints(TFunction<void(FLidarPointCloudPoint*)> Action, const bool& bVisibleOnly)
 {
 	PROCESS_ALL({ Action(Point); });
 }
 
-void FLidarPointCloudOctree::ExecuteActionOnAllPointsAsync(TFunction<void(FLidarPointCloudPoint*)> Action, const bool& bVisibleOnly, TFunction<void(void)> CompletionCallback /*= nullptr*/)
-{
-	Async(EAsyncExecution::Thread, [this, Action, bVisibleOnly]
-	{
-		ExecuteActionOnAllPointsAsync(Action, bVisibleOnly);
-	}, MoveTemp(CompletionCallback));
-}
-
 void FLidarPointCloudOctree::ExecuteActionOnPointsInSphere(TFunction<void(FLidarPointCloudPoint*)> Action, const FSphere& Sphere, const bool& bVisibleOnly)
 {
 	PROCESS_IN_SPHERE({ Action(Point); });
 }
 
-void FLidarPointCloudOctree::ExecuteActionOnPointsInSphereAsync(TFunction<void(FLidarPointCloudPoint*)> Action, const FSphere& Sphere, const bool& bVisibleOnly, TFunction<void(void)> CompletionCallback)
-{
-	Async(EAsyncExecution::Thread, [this, Action, Sphere, bVisibleOnly]
-	{
-		ExecuteActionOnPointsInSphere(Action, Sphere, bVisibleOnly);
-	}, MoveTemp(CompletionCallback));
-}
-
 void FLidarPointCloudOctree::ExecuteActionOnPointsInBox(TFunction<void(FLidarPointCloudPoint*)> Action, const FBox& Box, const bool& bVisibleOnly)
 {
 	PROCESS_IN_BOX({ Action(Point); });
 }
 
-void FLidarPointCloudOctree::ExecuteActionOnPointsInBoxAsync(TFunction<void(FLidarPointCloudPoint*)> Action, const FBox& Box, const bool& bVisibleOnly, TFunction<void(void)> CompletionCallback)
-{
-	Async(EAsyncExecution::Thread, [this, Action, Box, bVisibleOnly]
-	{
-		ExecuteActionOnPointsInBox(Action, Box, bVisibleOnly);
-	}, MoveTemp(CompletionCallback));
-}
-
 void FLidarPointCloudOctree::ExecuteActionOnPointsByRay(TFunction<void(FLidarPointCloudPoint*)> Action, const FLidarPointCloudRay& Ray, const float& Radius, bool bVisibleOnly)
 {
 	PROCESS_BY_RAY({ Action(Point); });
 }
 
-void FLidarPointCloudOctree::ExecuteActionOnPointsByRayAsync(TFunction<void(FLidarPointCloudPoint*)> Action, const FLidarPointCloudRay& Ray, const float& Radius, bool bVisibleOnly, TFunction<void(void)> CompletionCallback)
-{
-	Async(EAsyncExecution::Thread, [this, Action, Ray, Radius, bVisibleOnly]
-	{
-		ExecuteActionOnPointsByRay(Action, Ray, Radius, bVisibleOnly);
-	}, MoveTemp(CompletionCallback));
-}
-
 void FLidarPointCloudOctree::ApplyColorToAllPoints(const FColor& NewColor, const bool& bVisibleOnly)
 {
 	PROCESS_ALL({ Point->Color = NewColor; });
@@ -1258,123 +1235,62 @@ void FLidarPointCloudOctree::ApplyColorToPointsByRay(const FColor& NewColor, con
 
 void FLidarPointCloudOctree::MarkPointVisibilityDirty()
 {
-	TQueue<FLidarPointCloudOctreeNode*> Nodes;
-	FLidarPointCloudOctreeNode* CurrentNode = nullptr;
-	Nodes.Enqueue(&Root);
-	while (Nodes.Dequeue(CurrentNode))
-	{
-		CurrentNode->bVisibilityDirty = true;
-
-		for (auto Child : CurrentNode->Children)
-		{
-			Nodes.Enqueue(Child);
-		}
-	}
+	ITERATE_NODES({ CurrentNode->bVisibilityDirty = true; }, true);
 }
 
-void FLidarPointCloudOctree::Initialize(const FBox& InBounds)
+void FLidarPointCloudOctree::Initialize(const FVector& InExtent)
 {
-	if (!SetNewBounds(InBounds))
+	const bool bValidExtent = InExtent.X > 0 && InExtent.Y > 0 && InExtent.Z > 0;
+	if (!bValidExtent)
 	{
+		PC_ERROR("Provided bounds are incorrect: %s", *InExtent.ToString());
 		return;
 	}
 
+	Extent = InExtent;
+	const FVector UniformExtent = FVector(InExtent.GetMax());
+	
 	MaxBucketSize = GetDefault<ULidarPointCloudSettings>()->MaxBucketSize;
 	NodeGridResolution = GetDefault<ULidarPointCloudSettings>()->NodeGridResolution;
 
 	// Pre-calculate the shared per LOD data
 	for (int32 i = 0; i < SharedData.Num(); i++)
 	{
-		SharedData[i] = FSharedLODData(UniformBounds.GetExtent() / FMath::Pow(2, i));
+		SharedData[i] = FSharedLODData(UniformExtent / FMath::Pow(2, i));
 		NodeCount[i].Reset();
 		PointCount[i].Reset();
 	}
 
 	Empty(true);
+
+	bIsFullyLoaded = false;
 }
 
-void FLidarPointCloudOctree::ShiftPointsBy(FDoubleVector Offset, bool bRefreshPointsBounds)
+void FLidarPointCloudOctree::RemovePoint(const FLidarPointCloudPoint* Point)
 {
-	TQueue<FLidarPointCloudOctreeNode*> Nodes;
-	FLidarPointCloudOctreeNode* CurrentNode = nullptr;
-	Nodes.Enqueue(&Root);
-	while (Nodes.Dequeue(CurrentNode))
-	{
-		CurrentNode->Center = (Offset + CurrentNode->Center).ToVector();
-
-		FOR(Point, CurrentNode->AllocatedPoints)
-		{
-			Point->Location = (Offset + Point->Location).ToVector();
-		}
-
-		FOR(Point, CurrentNode->PaddingPoints)
-		{
-			Point->Location = (Offset + Point->Location).ToVector();
-		}
-
-		for (auto Child : CurrentNode->Children)
-		{
-			Nodes.Enqueue(Child);
-		}
-	}
-
-	Bounds.Min = (Offset + Bounds.Min).ToVector();
-	Bounds.Max = (Offset + Bounds.Max).ToVector();
-	UniformBounds.Min = (Offset + UniformBounds.Min).ToVector();
-	UniformBounds.Max = (Offset + UniformBounds.Max).ToVector();
-
-	MarkTraversalOctreesForInvalidation();
-
-	if (bRefreshPointsBounds)
-	{
-		RefreshPointsBounds();
-	}
-}
-
-void FLidarPointCloudOctree::RemovePoint(const FLidarPointCloudPoint* Point, bool bRefreshPointsBounds)
-{
-	if (!Point)
+	if (!Point)
 	{
 		return;
 	}
 	
 	int32 Index = -1;
-	bool bAllocatedPoint = false;
 	FLidarPointCloudOctreeNode* CurrentNode = &Root;
 	while (CurrentNode)
 	{
-		// Scan Allocated Points
+		const FLidarPointCloudPoint* RESTRICT Start = CurrentNode->GetData();
+		const int32 NumElements = CurrentNode->GetNumPoints();
+		for (const FLidarPointCloudPoint* RESTRICT Data = Start, *RESTRICT DataEnd = Start + NumElements; Data != DataEnd; ++Data)
 		{
-			const FLidarPointCloudPoint* Start = CurrentNode->AllocatedPoints.GetData();
-			const int32 NumElements = CurrentNode->AllocatedPoints.Num();
-			for (const FLidarPointCloudPoint* Data = Start, *DataEnd = Start + NumElements; Data != DataEnd; ++Data)
+			if (Data == Point)
 			{
-				if (Data == Point)
-				{
-					Index = Data - Start;
-					bAllocatedPoint = true;
-					break;
-				}
-			}
-		}
-
-		// Scan Padding Points
-		{
-			const FLidarPointCloudPoint* Start = CurrentNode->PaddingPoints.GetData();
-			const int32 NumElements = CurrentNode->PaddingPoints.Num();
-			for (const FLidarPointCloudPoint* Data = Start, *DataEnd = Start + NumElements; Data != DataEnd; ++Data)
-			{
-				if (Data == Point)
-				{
-					Index = Data - Start;
-					break;
-				}
+				Index = Data - Start;
+				break;
 			}
 		}
 
-		if (Index > 0)
+		if (Index > INDEX_NONE)
 		{
-			RemovePoint_Internal(CurrentNode, Index, bAllocatedPoint);
+			RemovePoint_Internal(CurrentNode, Index);
 			break;
 		}
 		else
@@ -1384,31 +1300,29 @@ void FLidarPointCloudOctree::RemovePoint(const FLidarPointCloudPoint* Point, boo
 		}
 	}
 
-	if (bRefreshPointsBounds)
-	{
-		RefreshPointsBounds();
-	}
+	RefreshBounds();
 }
 
-void FLidarPointCloudOctree::RemovePoint(FLidarPointCloudPoint Point, bool bRefreshPointsBounds)
+void FLidarPointCloudOctree::RemovePoint(FLidarPointCloudPoint Point)
 {
-	bool bAllocatedPoint = false;
+	int32 Index = -1;
 	FLidarPointCloudOctreeNode* CurrentNode = &Root;
 	while (CurrentNode)
 	{
-		int32 Index = CurrentNode->AllocatedPoints.IndexOfByKey(Point);
-		if (Index != INDEX_NONE)
+		const FLidarPointCloudPoint* RESTRICT Start = CurrentNode->GetData();
+		const int32 NumElements = CurrentNode->GetNumPoints();
+		for (const FLidarPointCloudPoint* RESTRICT Data = Start, *RESTRICT DataEnd = Start + NumElements; Data != DataEnd; ++Data)
 		{
-			bAllocatedPoint = true;
-		}
-		else
-		{
-			Index = CurrentNode->PaddingPoints.IndexOfByKey(Point);
+			if (*Data == Point)
+			{
+				Index = Data - Start;
+				break;
+			}
 		}
 
 		if (Index != INDEX_NONE)
 		{
-			RemovePoint_Internal(CurrentNode, Index, bAllocatedPoint);
+			RemovePoint_Internal(CurrentNode, Index);
 			break;
 		}
 		else
@@ -1418,153 +1332,150 @@ void FLidarPointCloudOctree::RemovePoint(FLidarPointCloudPoint Point, bool bRefr
 		}
 	}
 
-	if (bRefreshPointsBounds)
-	{
-		RefreshPointsBounds();
-	}
+	RefreshBounds();
 }
 
-void FLidarPointCloudOctree::RemovePoints(TArray<FLidarPointCloudPoint*>& Points, bool bRefreshPointsBounds)
+void FLidarPointCloudOctree::RemovePoints(TArray<FLidarPointCloudPoint*>& Points)
 {
 	if (Points.Num() == 0)
 	{
 		return;
 	}
 
-	FOR(Point, Points)
+	for (auto Point = Points.GetData(), DataEnd = Point + Points.Num(); Point != DataEnd; ++Point)
 	{
 		(*Point)->bMarkedForDeletion = true;
 	}
 
 	ITERATE_NODES(
 	{
-		int64 NumRemoved = 0;
-
-		// Scan Allocated Points
+		bool bHasPointsToRemove = false;
 		{
-			FLidarPointCloudPoint* Start = CurrentNode->AllocatedPoints.GetData();
-			const int32 NumElements = CurrentNode->AllocatedPoints.Num();
+			FLidarPointCloudPoint* Start = CurrentNode->GetData();
+			const int32 NumElements = CurrentNode->GetNumPoints();
 			for (FLidarPointCloudPoint* Data = Start, *DataEnd = Start + NumElements; Data != DataEnd; ++Data)
 			{
 				if (Data->bMarkedForDeletion)
 				{
-					CurrentNode->AllocatedPoints.RemoveAtSwap(Data - Start, 1, false);
-					++NumRemoved;
-					--DataEnd; 
-					--Data;
+					bHasPointsToRemove = true;
+					break;
 				}
 			}
 		}
 
-		// Scan Padding Points
+		if (bHasPointsToRemove)
 		{
-			FLidarPointCloudPoint* Start = CurrentNode->PaddingPoints.GetData();
-			const int32 NumElements = CurrentNode->PaddingPoints.Num();
+			int64 NumRemoved = 0;
+
+			// Make a copy of the data
+			TArray<FLidarPointCloudPoint> AllocatedPoints;
+			CurrentNode->BulkData.CopyToArray(AllocatedPoints);
+
+			FLidarPointCloudPoint* Start = AllocatedPoints.GetData();
+			const int32 NumElements = AllocatedPoints.Num();
 			for (FLidarPointCloudPoint* Data = Start, *DataEnd = Start + NumElements; Data != DataEnd; ++Data)
 			{
 				if (Data->bMarkedForDeletion)
 				{
-					CurrentNode->PaddingPoints.RemoveAtSwap(Data - Start, 1, false);
+					AllocatedPoints.RemoveAtSwap(Data - Start, 1, false);
 					++NumRemoved;
 					--DataEnd;
 					--Data;
 				}
 			}
-		}
 
-		if (NumRemoved > 0)
-		{
 			// #todo: Fetch points from child nodes / padding points to fill the gap
-			
+
 			CurrentNode->AddPointCount(this, -NumRemoved);
 
 			// Reduce space usage
-			CurrentNode->AllocatedPoints.Shrink();
-			CurrentNode->PaddingPoints.Shrink();
+			AllocatedPoints.Shrink();
+
+			// Copy the updated array back to the BulkData
+			CurrentNode->BulkData.CopyFromArray(AllocatedPoints);
+			CurrentNode->bCanReleaseData = false;
 
 			// Sort points to speed up rendering
-			Algo::Sort(CurrentNode->AllocatedPoints, MoveTemp(PointVisibilitySortingFunction));
-			Algo::Sort(CurrentNode->PaddingPoints, MoveTemp(PointVisibilitySortingFunction));
+			CurrentNode->SortVisiblePoints();
 		}
 	}, true);
 
-	if (bRefreshPointsBounds)
-	{
-		RefreshPointsBounds();
-	}
+	RefreshBounds();
 }
 
-void FLidarPointCloudOctree::RemovePointsInSphere(const FSphere& Sphere, const bool& bVisibleOnly, bool bRefreshPointsBounds)
+void FLidarPointCloudOctree::RemovePointsInSphere(const FSphere& Sphere, const bool& bVisibleOnly)
 {
 	// #todo: This can be optimized by removing points inline
 	TArray<FLidarPointCloudPoint*> SelectedPoints;
 	GetPointsInSphere(SelectedPoints, Sphere, bVisibleOnly);
-	RemovePoints(SelectedPoints, bRefreshPointsBounds);
+	RemovePoints(SelectedPoints);
 }
 
-void FLidarPointCloudOctree::RemovePointsInBox(const FBox& Box, const bool& bVisibleOnly, bool bRefreshPointsBounds)
+void FLidarPointCloudOctree::RemovePointsInBox(const FBox& Box, const bool& bVisibleOnly)
 {
 	// #todo: This can be optimized by removing points inline
 	TArray<FLidarPointCloudPoint*> SelectedPoints;
 	GetPointsInBox(SelectedPoints, Box, bVisibleOnly);
-	RemovePoints(SelectedPoints, bRefreshPointsBounds);
+	RemovePoints(SelectedPoints);
 }
 
-void FLidarPointCloudOctree::RemovePointsByRay(const FLidarPointCloudRay& Ray, const float& Radius, const bool& bVisibleOnly, bool bRefreshPointsBounds)
+void FLidarPointCloudOctree::RemovePointsByRay(const FLidarPointCloudRay& Ray, const float& Radius, const bool& bVisibleOnly)
 {
 	// #todo: This can be optimized by removing points inline
 	TArray<FLidarPointCloudPoint*> SelectedPoints;
 	RaycastMulti(Ray, Radius, bVisibleOnly, SelectedPoints);
-	RemovePoints(SelectedPoints, bRefreshPointsBounds);
+	RemovePoints(SelectedPoints);
 }
 
-void FLidarPointCloudOctree::RemoveHiddenPoints(bool bRefreshPointsBounds)
+void FLidarPointCloudOctree::RemoveHiddenPoints()
 {
 	ITERATE_NODES(
 	{
-		int64 NumRemoved = 0;
-
-		// Scan Allocated Points
+		bool bHasPointsToRemove = false;
 		{
-			FLidarPointCloudPoint* Start = CurrentNode->AllocatedPoints.GetData();
-			const int32 NumElements = CurrentNode->AllocatedPoints.Num();
+			FLidarPointCloudPoint* Start = CurrentNode->GetData();
+			const int32 NumElements = CurrentNode->GetNumPoints();
 			for (FLidarPointCloudPoint* Data = Start, *DataEnd = Start + NumElements; Data != DataEnd; ++Data)
 			{
 				if (!Data->bVisible)
 				{
-					CurrentNode->AllocatedPoints.RemoveAtSwap(Data - Start, 1, false);
-					++NumRemoved;
-					--DataEnd; 
-					--Data;
+					bHasPointsToRemove = true;
+					break;
 				}
 			}
 		}
 
-		// Scan Padding Points
+		if (bHasPointsToRemove)
 		{
-			FLidarPointCloudPoint* Start = CurrentNode->PaddingPoints.GetData();
-			const int32 NumElements = CurrentNode->PaddingPoints.Num();
+			int64 NumRemoved = 0;
+
+			// Make a copy of the data
+			TArray<FLidarPointCloudPoint> AllocatedPoints;
+			CurrentNode->BulkData.CopyToArray(AllocatedPoints);
+
+			FLidarPointCloudPoint* Start = AllocatedPoints.GetData();
+			const int32 NumElements = AllocatedPoints.Num();
 			for (FLidarPointCloudPoint* Data = Start, *DataEnd = Start + NumElements; Data != DataEnd; ++Data)
 			{
 				if (!Data->bVisible)
 				{
-					CurrentNode->PaddingPoints.RemoveAtSwap(Data - Start, 1, false);
+					AllocatedPoints.RemoveAtSwap(Data - Start, 1, false);
 					++NumRemoved;
 					--DataEnd;
 					--Data;
 				}
 			}
-		}
 
-		if (NumRemoved > 0)
-		{
 			// #todo: Fetch points from child nodes / padding points to fill the gap
-			
+
 			CurrentNode->AddPointCount(this, -NumRemoved);
 
 			// Reduce space usage
-			CurrentNode->AllocatedPoints.Shrink();
-			CurrentNode->PaddingPoints.Shrink();
+			AllocatedPoints.Shrink();
+
+			// Copy the updated array back to the BulkData
+			CurrentNode->BulkData.CopyFromArray(AllocatedPoints);
+			CurrentNode->bCanReleaseData = false;
 
 			// Set visibility data
 			CurrentNode->NumVisiblePoints = CurrentNode->GetNumPoints();
@@ -1572,10 +1483,7 @@ void FLidarPointCloudOctree::RemoveHiddenPoints(bool bRefreshPointsBounds)
 		}
 	}, true);
 
-	if (bRefreshPointsBounds)
-	{
-		RefreshPointsBounds();
-	}
+	RefreshBounds();
 }
 
 void FLidarPointCloudOctree::Empty(bool bDestroyNodes)
@@ -1590,7 +1498,10 @@ void FLidarPointCloudOctree::Empty(bool bDestroyNodes)
 			Count.Reset();
 		}
 
-		new (&Root) FLidarPointCloudOctreeNode(this, nullptr, 0, 0, Bounds.GetCenter());
+		new (&Root) FLidarPointCloudOctreeNode(this, 0);
+
+		QueuedNodes.Empty();
+		NodesInUse.Reset();
 
 		MarkTraversalOctreesForInvalidation();
 	}
@@ -1598,9 +1509,7 @@ void FLidarPointCloudOctree::Empty(bool bDestroyNodes)
 	{
 		Root.Empty(true);
 	}
-
-	PointsBounds.Init();
-
+	
 	// Reset point counters
 	for (auto& Count : PointCount)
 	{
@@ -1608,27 +1517,158 @@ void FLidarPointCloudOctree::Empty(bool bDestroyNodes)
 	}
 }
 
-bool FLidarPointCloudOctree::SetNewBounds(const FBox& InBounds)
+void FLidarPointCloudOctree::UnregisterTraversalOctree(FLidarPointCloudTraversalOctree* TraversalOctree)
 {
-	if (!InBounds.IsValid)
+	if (TraversalOctree)
 	{
-		return false;
+		bool bRemoved = false;
+
+		for (int32 i = 0; i < LinkedTraversalOctrees.Num(); ++i)
+		{
+			if (auto TO = LinkedTraversalOctrees[i].Pin())
+			{
+				if (TO.Get() == TraversalOctree)
+				{
+					LinkedTraversalOctrees.RemoveAtSwap(i--);
+					bRemoved = true;
+				}
+			}
+			// Remove null
+			else
+			{
+				LinkedTraversalOctrees.RemoveAtSwap(i--);
+			}
+		}
+
+		// If nothing is using this Octree, release all non-persistent nodes
+		if (bRemoved && LinkedTraversalOctrees.Num() == 0)
+		{
+			FScopeLock Lock(&DataLock);
+			ReleaseAllNodes(false);
+		}
 	}
+}
 
-	Bounds = InBounds;
+void FLidarPointCloudOctree::QueueNode(FLidarPointCloudOctreeNode* Node, float Lifetime)
+{
+	if (!Node)
+	{
+		return;
+	}
 
-	FVector BoundsExtent = FVector(Bounds.GetExtent().GetMax());
-	UniformBounds = FBox(Bounds.GetCenter() - BoundsExtent, Bounds.GetCenter() + BoundsExtent);
+	// Refresh lifetime of the BulkData, if requested
+	if (Lifetime > -1)
+	{
+		Node->BulkDataLifetime = Lifetime;
+	}
 
-	return true;
+	// No need to do anything, if the node already has data loaded or loading
+	if (Node->HasData() || Node->bHasDataPending)
+	{
+		return;
+	}
+
+	NodesInUse.Add(Node);
+	QueuedNodes.Enqueue(Node);
+	Node->bHasDataPending = true;
 }
 
-void FLidarPointCloudOctree::RemovePoint_Internal(FLidarPointCloudOctreeNode* Node, int32 Index, bool bAllocatedPoint)
+void FLidarPointCloudOctree::StreamQueuedNodes()
+{
+	SCOPE_CYCLE_COUNTER(STAT_NodeStreaming);
+
+	// Only one streaming operation at a time
+	if (bStreamingBusy)
+	{
+		return;
+	}
+
+	bStreamingBusy = true;
+	
+	// Perform data streaming in a separate thread
+	Async(EAsyncExecution::TaskGraph, [this]
+	{
+		SCOPE_CYCLE_COUNTER(STAT_NodeStreaming);
+
+		FScopeLock Lock(&DataLock);
+
+		FLidarPointCloudOctreeNode* CurrentNode = nullptr;
+		while (QueuedNodes.Dequeue(CurrentNode))
+		{
+			CurrentNode->GetData();
+			CurrentNode->bHasDataPending = false;
+		}
+
+		bStreamingBusy = false;
+	});
+}
+
+void FLidarPointCloudOctree::UnloadOldNodes(const float& CurrentTime)
+{
+	SCOPE_CYCLE_COUNTER(STAT_NodeStreaming);
+
+	for(int32 i = 0; i < NodesInUse.Num(); ++i)
+	{
+		auto& Node = NodesInUse[i];
+
+		// Unload data, if it expired
+		if (Node->BulkDataLifetime < CurrentTime)
+		{
+			Node->ReleaseData();
+			NodesInUse.RemoveAtSwap(i--, 1, false);
+		}
+	}
+}
+
+void FLidarPointCloudOctree::LoadAllNodes()
+{
+	ITERATE_NODES({ CurrentNode->GetPersistentData(); }, true);
+
+	bIsFullyLoaded = true;
+}
+
+void FLidarPointCloudOctree::ReleaseAllNodes(bool bIncludePersistent)
+{
+	ITERATE_NODES({ CurrentNode->ReleaseData(bIncludePersistent); }, true);
+
+	if (bIncludePersistent)
+	{
+		bIsFullyLoaded = false;
+	}
+}
+
+void FLidarPointCloudOctree::RefreshAllocatedSize()
+{
+	FScopeTryLock Lock(&DataLock);
+	if (!Lock.IsLocked())
+	{
+		return;
+	}
+
+	PreviousPointCount = GetNumPoints();
+	PreviousNodeCount = GetNumNodes();
+
+	PreviousAllocatedStructureSize = sizeof(FLidarPointCloudOctree);
+
+	PreviousAllocatedStructureSize += SharedData.GetAllocatedSize();
+	PreviousAllocatedStructureSize += PointCount.GetAllocatedSize();
+
+	PreviousAllocatedSize = PreviousAllocatedStructureSize + Root.GetAllocatedSize(true, true);
+	PreviousAllocatedStructureSize += Root.GetAllocatedSize(true, false);
+}
+
+void FLidarPointCloudOctree::RemovePoint_Internal(FLidarPointCloudOctreeNode* Node, int32 Index)
 {
 	Node->AddPointCount(this, -1);
 
-	// Remove the point from relevant array
-	(bAllocatedPoint ? Node->AllocatedPoints : Node->PaddingPoints).RemoveAt(Index);
+	// Make a copy of the data
+	TArray<FLidarPointCloudPoint> AllocatedPoints;
+	Node->BulkData.CopyToArray(AllocatedPoints);
+
+	AllocatedPoints.RemoveAt(Index);
+
+	// Copy the updated array back to the BulkData
+	Node->BulkData.CopyFromArray(AllocatedPoints);
 
 	// #todo: Fetch points from child nodes / padding points to fill the gap
 }
@@ -1637,61 +1677,155 @@ void FLidarPointCloudOctree::MarkTraversalOctreesForInvalidation()
 {
 	for (int32 i = 0; i < LinkedTraversalOctrees.Num(); ++i)
 	{
-		auto TraversalOctree = LinkedTraversalOctrees[i];
-
-		// Remove null
-		if (!TraversalOctree)
+		if (auto TraversalOctree = LinkedTraversalOctrees[i].Pin())
 		{
-			LinkedTraversalOctrees.RemoveAtSwap(i--);
+			TraversalOctree->bValid = false;
 		}
-		else if(TraversalOctree->bValid)
+		// Remove null
+		else
 		{
-			TraversalOctree->bValid = false;
+			LinkedTraversalOctrees.RemoveAtSwap(i--);
 		}
 	}
 }
 
-FArchive& operator<<(FArchive& Ar, FLidarPointCloudOctree& O)
+void FLidarPointCloudOctree::Serialize(FArchive& Ar)
 {
-	Ar << O.Bounds;
+	// Extent
+	{
+		FVector NodesExtent = SharedData[0].Extent;
+
+		if (Ar.CustomVer(ULidarPointCloud::PointCloudFileGUID) > 16)
+		{
+			Ar << NodesExtent;
+		}
+		else
+		{
+			FBox Bounds;
+			Ar << Bounds;
+			NodesExtent = Bounds.GetExtent();
+		}
+
+		if (Ar.IsLoading())
+		{
+			Initialize(NodesExtent);
+		}
+	}
 
 	// Collision Mesh data
 	if (Ar.CustomVer(ULidarPointCloud::PointCloudFileGUID) > 13)
 	{
 		FTriMeshCollisionData Dummy;
-		FTriMeshCollisionData* CollisionMesh = Ar.IsCooking() ? &Dummy : &O.CollisionMesh;
+		FTriMeshCollisionData* CollisionMeshPtr = Ar.IsCooking() ? &Dummy : &CollisionMesh;
 
-		Ar << CollisionMesh->Vertices;
+		Ar << CollisionMeshPtr->Vertices;
 
-		int32 NumIndices = CollisionMesh->Indices.Num();
+		int32 NumIndices = CollisionMeshPtr->Indices.Num();
 		Ar << NumIndices;
 
 		if (Ar.IsLoading())
 		{
-			CollisionMesh->Indices.AddUninitialized(NumIndices);
+			CollisionMeshPtr->Indices.AddUninitialized(NumIndices);
 		}
 
-		Ar.Serialize(CollisionMesh->Indices.GetData(), NumIndices * sizeof(FTriIndices));
+		Ar.Serialize(CollisionMeshPtr->Indices.GetData(), NumIndices * sizeof(FTriIndices));
 	}
 
-	if (Ar.IsLoading())
+	const bool bUseCompression = GetDefault<ULidarPointCloudSettings>()->bUseCompression;
+
+	// Used for backwards compatibility with pre-streaming formats
+	if (Ar.CustomVer(ULidarPointCloud::PointCloudFileGUID) < 16)
 	{
-		O.Initialize(O.Bounds);
+		TArray<FLidarPointCloudOctreeNode*> Nodes;
+		Nodes.Add(&Root);
+		while (Nodes.Num())
+		{
+			FLidarPointCloudOctreeNode* CurrentNode = Nodes.Pop(false);
+
+			Ar << CurrentNode->LocationInParent << CurrentNode->Center;
+			CurrentNode->BulkData.SerializeLegacy(Ar);
+			CurrentNode->bCanReleaseData = false;
+
+			int32 NumChildren = CurrentNode->Children.Num();
+			Ar << NumChildren;
+
+			CurrentNode->AddPointCount(this);
+
+			CurrentNode->Children.AddUninitialized(NumChildren);
+			for (int32 i = 0; i < NumChildren; ++i)
+			{
+				CurrentNode->Children[i] = new FLidarPointCloudOctreeNode(this, CurrentNode->Depth + 1);
+			}
+
+			for (int32 i = NumChildren - 1; i >= 0; --i)
+			{
+				Nodes.Add(CurrentNode->Children[i]);
+			}
+		}
 	}
+	else
+	{
+		ITERATE_NODES({
+			if (Ar.IsSaving())
+			{
+				CurrentNode->BulkData.ClearBulkDataFlags(BULKDATA_SerializeCompressed);
+				if (bUseCompression)
+				{
+					CurrentNode->BulkData.SetBulkDataFlags(BULKDATA_SerializeCompressed);
+				}
+
+				// Make sure the points are in optimized order before saving
+				CurrentNode->SortVisiblePoints();
+			}
 
-	O.Root.Serialize(Ar, &O);
+			CurrentNode->BulkData.Serialize(Ar, Owner);
+
+			if (Ar.IsSaving())
+			{
+				CurrentNode->bCanReleaseData = true;
+			}
+
+			Ar << CurrentNode->LocationInParent << CurrentNode->Center;
+			int32 NumChildren = CurrentNode->Children.Num();
+			Ar << NumChildren;
+
+			if (Ar.IsLoading())
+			{
+				CurrentNode->AddPointCount(this);
+
+				CurrentNode->Children.AddUninitialized(NumChildren);
+				for (int32 i = 0; i < NumChildren; i++)
+				{
+					CurrentNode->Children[i] = new FLidarPointCloudOctreeNode(this, CurrentNode->Depth + 1);
+				}
+			}
+		}, true);
+	}
 
-	// Points Bounds
-	if (Ar.CustomVer(ULidarPointCloud::PointCloudFileGUID) > 14)
+	// Points Extent
+	if (Ar.CustomVer(ULidarPointCloud::PointCloudFileGUID) > 16)
 	{
-		Ar << O.PointsBounds;
+		Ar << Extent;
 	}
 	else
 	{
-		O.RefreshPointsBounds();
-	}
+		FBox PointsBounds;
 
-	return Ar;
+		if (Ar.CustomVer(ULidarPointCloud::PointCloudFileGUID) > 14)
+		{
+			Ar << PointsBounds;
+		}
+
+		if (Ar.CustomVer(ULidarPointCloud::PointCloudFileGUID) > 15)
+		{
+			Extent = PointsBounds.GetExtent();
+			Owner->LocationOffset = PointsBounds.GetCenter();
+		}
+		else
+		{
+			RefreshBounds();
+		}
+	}
 }
 
 //////////////////////////////////////////////////////////// FLidarPointCloudTraversalOctreeNode
@@ -1702,10 +1836,10 @@ FLidarPointCloudTraversalOctreeNode::FLidarPointCloudTraversalOctreeNode()
 {
 }
 
-void FLidarPointCloudTraversalOctreeNode::Build(FLidarPointCloudOctreeNode* Node, const FTransform& LocalToWorld)
+void FLidarPointCloudTraversalOctreeNode::Build(FLidarPointCloudOctreeNode* Node, const FTransform& LocalToWorld, const FVector& LocationOffset)
 {
 	DataNode = Node;
-	Center = LocalToWorld.TransformPosition(Node->Center);
+	Center = LocalToWorld.TransformPosition(Node->Center + LocationOffset);
 	Depth = Node->Depth;
 
 	Children.AddZeroed(Node->Children.Num());
@@ -1713,7 +1847,7 @@ void FLidarPointCloudTraversalOctreeNode::Build(FLidarPointCloudOctreeNode* Node
 	{
 		if (Node->Children[i])
 		{
-			Children[i].Build(Node->Children[i], LocalToWorld);
+			Children[i].Build(Node->Children[i], LocalToWorld, LocationOffset);
 			Children[i].Parent = this;
 		}
 	}
@@ -1777,162 +1911,52 @@ void FLidarPointCloudTraversalOctreeNode::CalculateVirtualDepth(const TArray<flo
 FLidarPointCloudTraversalOctree::FLidarPointCloudTraversalOctree(FLidarPointCloudOctree* Octree, const FTransform& LocalToWorld)
 	: Octree(Octree)
 	, bValid(false)
-{
-	Octree->RegisterTraversalOctree(this);
-	
+{	
 	// Reset the tree
 	Extents.Empty();
 	RadiiSq.Empty();
-	PointCount.Empty();
-	NumPoints = 0;
 	Root.~FLidarPointCloudTraversalOctreeNode();
 	new (&Root) FLidarPointCloudTraversalOctreeNode();
 
 	// Calculate properties
 	NumLODs = Octree->GetNumLODs();
 
-	FBox WorldBounds = Octree->UniformBounds.TransformBy(LocalToWorld);
+	VirtualDepthMultiplier = 255.0f / NumLODs;
+	ReversedVirtualDepthMultiplier = NumLODs / 255.0f;
+
+	const FVector Extent = Octree->SharedData[0].Extent;
+
+	FBox WorldBounds = FBox(-Extent, Extent).TransformBy(LocalToWorld);
 	for (int32 i = 0; i < NumLODs; i++)
 	{
 		Extents.Emplace(i == 0 ? WorldBounds.GetExtent() : Extents.Last() * 0.5f);
 		RadiiSq.Emplace(FMath::Square(Extents.Last().Size()));
 	}
 
+	int64 NumPoints = 0;
+	TArray<int64> PointCount;
 	for (auto& Count : Octree->PointCount)
 	{
 		PointCount.Add(Count.GetValue());
 		NumPoints += Count.GetValue();
 	}
 
-	// Star cloning the node data
-	Root.Build(&Octree->Root, LocalToWorld);
-
-	bValid = true;
-}
-
-FLidarPointCloudTraversalOctree::~FLidarPointCloudTraversalOctree()
-{
-	Octree->UnregisterTraversalOctree(this);
-}
-
-TArray<FLidarPointCloudTraversalOctreeNode*> FLidarPointCloudTraversalOctree::GetVisibleNodes(const FSceneView* View, const float& MinScreenSizeSq, const int32& PointBudget, const float& ScreenCenterImportance, const int32& MinDepth, const int32& MaxDepth)
-{
-	SCOPE_CYCLE_COUNTER(STAT_NodeSelection);
-
-	// For Top / Front / Side views disable ScreenSize checks
-	const bool bSkipMinScreenSize = !View->bIsGameView && !View->IsPerspectiveProjection();
-
-	// Collect Nodes
-	TArray<FNodeSizeData> NodeSizeData;
-	{
-		const FMatrix& ProjMatrix = View->ViewMatrices.GetProjectionMatrix();
-		const float ScreenSizeFactor = FMath::Square(FMath::Max(0.5f * ProjMatrix.M[0][0], 0.5f * ProjMatrix.M[1][1]));
-		const FVector& ViewOrigin = View->ViewMatrices.GetViewOrigin();
-
-		TQueue<FLidarPointCloudTraversalOctreeNode*> Nodes;
-		FLidarPointCloudTraversalOctreeNode* CurrentNode = nullptr;
-		Nodes.Enqueue(&Root);
-		while (Nodes.Dequeue(CurrentNode))
-		{
-			// Reset selection flag
-			CurrentNode->bSelected = false;
-
-			// Update number of visible points, if needed
-			CurrentNode->DataNode->UpdateNumVisiblePoints();
-
-			// In Frustum?
-			// #todo: Skip frustum checks for nodes fully in frustum
-			if (!View->ViewFrustum.IntersectBox(CurrentNode->Center, Extents[CurrentNode->Depth]))
-			{
-				continue;
-			}
-
-			// Only process this node if it has any visible points - do not use continue; as the children may still contain visible points!
-			if (CurrentNode->DataNode->GetNumVisiblePoints() > 0 && CurrentNode->Depth >= MinDepth)
-			{
-				float ScreenSizeSq = 0;
-
-				FVector VectorToNode = CurrentNode->Center - ViewOrigin;
-				const float DistSqr = VectorToNode.SizeSquared();
-
-				// If the camera is within this node's bounds, it should always be qualified for rendering
-				if (DistSqr <= RadiiSq[CurrentNode->Depth])
-				{
-					// Subtract Depth to maintain hierarchy 
-					ScreenSizeSq = 10000 - CurrentNode->Depth;
-				}
-				else
-				{
-					ScreenSizeSq = ScreenSizeFactor * RadiiSq[CurrentNode->Depth] / FMath::Max(1.0f, DistSqr);
-
-					// Check for minimum screen size
-					if (!bSkipMinScreenSize && ScreenSizeSq < MinScreenSizeSq)
-					{
-						continue;
-					}
-
-					// Add optional preferential selection for nodes closer to the screen center
-					if (ScreenCenterImportance > 0)
-					{
-						VectorToNode.Normalize();
-						float Dot = FVector::DotProduct(View->GetViewDirection(), VectorToNode);
-
-						ScreenSizeSq = FMath::Lerp(ScreenSizeSq, ScreenSizeSq * Dot, ScreenCenterImportance);
-					}
-				}
-
-				NodeSizeData.Emplace(CurrentNode, ScreenSizeSq);
-			}
-
-			if (MaxDepth < 0 || CurrentNode->Depth < MaxDepth)
-			{
-				for (auto& Child : CurrentNode->Children)
-				{
-					Nodes.Enqueue(&Child);
-				}
-			}
-		}
-	}
-
-	// Sort Nodes
-	Algo::Sort(NodeSizeData, [](const FNodeSizeData& A, const FNodeSizeData& B) { return A.Size > B.Size; });
-
-	// Select Nodes
-	int32 TotalPointsSelected = 0;
-	TArray<FLidarPointCloudTraversalOctreeNode*> SelectedNodes;
-
-	for (FNodeSizeData& Element : NodeSizeData)
+	LevelWeights.AddZeroed(NumLODs);
+	for (int32 i = 0; i < LevelWeights.Num(); i++)
 	{
-		int32 NewNumPointsSelected = TotalPointsSelected + Element.Node->DataNode->GetNumVisiblePoints();
-
-		if (NewNumPointsSelected <= PointBudget)
-		{
-			SelectedNodes.Add(Element.Node);
-			TotalPointsSelected = NewNumPointsSelected;
-			Element.Node->bSelected = true;
-		}
+		LevelWeights[i] = (float)PointCount[i] / NumPoints;
 	}
 
-	return SelectedNodes;
-}
+	// Star cloning the node data
+	Root.Build(&Octree->Root, LocalToWorld, Octree->Owner->GetLocationOffset().ToVector());
 
-void FLidarPointCloudTraversalOctree::CalculateSpriteSize(TArray<FLidarPointCloudTraversalOctreeNode*>& SelectedNodes, const float& PointSizeBias)
-{
-	for (auto Node : SelectedNodes)
-	{
-		Node->CalculateVirtualDepth(CalculateLevelWeights(), 255.0f / NumLODs, PointSizeBias);
-	}
+	bValid = true;
 }
 
-TArray<float> FLidarPointCloudTraversalOctree::CalculateLevelWeights() const
+FLidarPointCloudTraversalOctree::~FLidarPointCloudTraversalOctree()
 {
-	TArray<float> LevelWeights;
-	LevelWeights.AddZeroed(NumLODs);
-
-	for (int32 i = 0; i < LevelWeights.Num(); i++)
+	if (bValid)
 	{
-		LevelWeights[i] = (float)PointCount[i] / NumPoints;
+		Octree->UnregisterTraversalOctree(this);
 	}
-
-	return LevelWeights;
-}
+}
\ No newline at end of file
diff --git a/Engine/Plugins/Enterprise/LidarPointCloud/Source/LidarPointCloudRuntime/Private/LidarPointCloudOctreeMacros.h b/Engine/Plugins/Enterprise/LidarPointCloud/Source/LidarPointCloudRuntime/Private/LidarPointCloudOctreeMacros.h
--- a/Engine/Plugins/Enterprise/LidarPointCloud/Source/LidarPointCloudRuntime/Private/LidarPointCloudOctreeMacros.h
+++ b/Engine/Plugins/Enterprise/LidarPointCloud/Source/LidarPointCloudRuntime/Private/LidarPointCloudOctreeMacros.h
@@ -2,7 +2,8 @@
 
 #pragma once
 
-#define FOR(ElemName, Array) for(auto ElemName = Array.GetData(), DataEnd = ElemName + Array.Num(); ElemName != DataEnd; ++ElemName)
+#define FOR(ElemName, Node) for(auto ElemName = Node->GetPersistentData(), DataEnd = ElemName + Node->GetNumPoints(); ElemName != DataEnd; ++ElemName)
+#define FOR_RO(ElemName, Node) for(auto ElemName = Node->GetData(), DataEnd = ElemName + Node->GetNumPoints(); ElemName != DataEnd; ++ElemName)
 
 #define IS_VIS_CHECK_REQUIRED (bVisibleOnly && CurrentNode->NumVisiblePoints < CurrentNode->GetNumPoints())
 
@@ -23,77 +24,71 @@
 #define ITERATE_NODES(Action, NodeTest) ITERATE_NODES_BODY(Action, NodeTest,  )
 #define ITERATE_NODES_CONST(Action, NodeTest) ITERATE_NODES_BODY(Action, NodeTest, const)
 
-#define ITERATE_POINTS(Action) \
-{\
-	FOR(Point, CurrentNode->AllocatedPoints) { Action }\
-	FOR(Point, CurrentNode->PaddingPoints) { Action }\
-}
-
 #define POINT_IN_BOX Box.IsInsideOrOn(Point->Location)
 #define POINT_IN_SPHERE (POINT_IN_BOX && FVector::DistSquared(Point->Location, Sphere.Center) <= RadiusSq)
 #define POINT_IN_FRUSTUM Frustum.IntersectSphere(Point->Location, 0)
 #define POINT_BY_RAY Ray.Intersects(Point, RadiusSq)
 
-#define PROCESS_BODY(Action, PointTest) \
+#define PROCESS_BODY(Action, PointTest, Mode) \
 {\
 	/* If node fully inside the box - do not check individual points */\
 	if (bNodeFullyContained)\
 	{\
-		if (!IS_VIS_CHECK_REQUIRED) { ITERATE_POINTS(Action) }\
-		else { ITERATE_POINTS({ if (Point->bVisible) { Action } }) }\
+		if (!IS_VIS_CHECK_REQUIRED) { FOR##Mode(Point, CurrentNode) { Action } }\
+		else { FOR##Mode(Point, CurrentNode) { if (Point->bVisible) { Action } } }\
 	}\
 	else\
 	{\
-		if (!IS_VIS_CHECK_REQUIRED) { ITERATE_POINTS({ if (PointTest) { Action } }) }\
-		else { ITERATE_POINTS({ if (Point->bVisible && PointTest) { Action } }) }\
+		if (!IS_VIS_CHECK_REQUIRED) { FOR##Mode(Point, CurrentNode) { if (PointTest) { Action } } }\
+		else { FOR##Mode(Point, CurrentNode) { if (Point->bVisible && PointTest) { Action } } }\
 	}\
 }
 
-#define PROCESS_ALL_BODY(Action) \
+#define PROCESS_ALL_BODY(Action, Mode) \
 {\
 	if (!bVisibleOnly || CurrentNode->NumVisiblePoints > 0)\
 	{\
 		const bool bNodeFullyContained = true;\
-		PROCESS_BODY(Action, true) \
+		PROCESS_BODY(Action, true, Mode) \
 	}\
 }
 
-#define PROCESS_IN_SPHERE_BODY(Action) \
+#define PROCESS_IN_SPHERE_BODY(Action, Mode) \
 {\
 	if (!bVisibleOnly || CurrentNode->NumVisiblePoints > 0)\
 	{\
 		const bool bNodeFullyContained = CurrentNode->GetSphereBounds(this).IsInside(Sphere);\
-		PROCESS_BODY(Action, POINT_IN_SPHERE) \
+		PROCESS_BODY(Action, POINT_IN_SPHERE, Mode) \
 	}\
 }
 
-#define PROCESS_IN_BOX_BODY(Action) \
+#define PROCESS_IN_BOX_BODY(Action, Mode) \
 {\
 	if (!bVisibleOnly || CurrentNode->NumVisiblePoints > 0)\
 	{\
 		const bool bNodeFullyContained = Box.IsInsideOrOn(CurrentNode->Center - SharedData[CurrentNode->Depth].Extent) && Box.IsInsideOrOn(CurrentNode->Center + SharedData[CurrentNode->Depth].Extent);\
-		PROCESS_BODY(Action, POINT_IN_BOX) \
+		PROCESS_BODY(Action, POINT_IN_BOX, Mode) \
 	}\
 }
 
-#define PROCESS_IN_FRUSTUM_BODY(Action) \
+#define PROCESS_IN_FRUSTUM_BODY(Action, Mode) \
 {\
 	if (!bVisibleOnly || CurrentNode->NumVisiblePoints > 0)\
 	{\
 		bool bNodeFullyContained;\
 		Frustum.IntersectBox(CurrentNode->Center, SharedData[CurrentNode->Depth].Extent, bNodeFullyContained);\
-		PROCESS_BODY(Action, POINT_IN_FRUSTUM) \
+		PROCESS_BODY(Action, POINT_IN_FRUSTUM, Mode) \
 	}\
 }
 
-#define PROCESS_BY_RAY_BODY(Action) \
+#define PROCESS_BY_RAY_BODY(Action, Mode) \
 {\
 	if (!bVisibleOnly || CurrentNode->NumVisiblePoints > 0)\
 	{\
 		if (Ray.Intersects(CurrentNode->GetBounds(this)))\
 		{\
-			if (!IS_VIS_CHECK_REQUIRED) { ITERATE_POINTS({ if (POINT_BY_RAY) { Action } }) }\
-			else { ITERATE_POINTS({ if (Point->bVisible && POINT_BY_RAY) { Action } }) }\
+			if (!IS_VIS_CHECK_REQUIRED) { FOR##Mode(Point, CurrentNode) { if (POINT_BY_RAY) { Action } } }\
+			else { FOR##Mode(Point, CurrentNode) { if (Point->bVisible && POINT_BY_RAY) { Action } } }\
 \
 			for (auto Child : CurrentNode->Children) Nodes.Enqueue(Child);\
 		}\
@@ -105,32 +100,32 @@
 	/* Build a box to quickly filter out the points - (IsInsideOrOn vs comparing DistSquared) */\
 	const FBox Box(Sphere.Center - FVector(Sphere.W), Sphere.Center + FVector(Sphere.W));\
 	const float RadiusSq = Sphere.W * Sphere.W;\
-	ITERATE_NODES(PROCESS_IN_SPHERE_BODY(Action), NODE_IN_BOX) \
+	ITERATE_NODES(PROCESS_IN_SPHERE_BODY(Action,), NODE_IN_BOX) \
 }
 #define PROCESS_IN_SPHERE_CONST(Action) \
 {\
 	/* Build a box to quickly filter out the points - (IsInsideOrOn vs comparing DistSquared) */\
 	const FBox Box(Sphere.Center - FVector(Sphere.W), Sphere.Center + FVector(Sphere.W));\
 	const float RadiusSq = Sphere.W * Sphere.W;\
-	ITERATE_NODES_CONST(PROCESS_IN_SPHERE_BODY(Action), NODE_IN_BOX) \
+	ITERATE_NODES_CONST(PROCESS_IN_SPHERE_BODY(Action, _RO), NODE_IN_BOX) \
 }
 
-#define PROCESS_ALL(Action) { ITERATE_NODES(PROCESS_ALL_BODY(Action), true) }
-#define PROCESS_ALL_CONST(Action) { ITERATE_NODES_CONST(PROCESS_ALL_BODY(Action), true) }
+#define PROCESS_ALL(Action) { ITERATE_NODES(PROCESS_ALL_BODY(Action,), true) }
+#define PROCESS_ALL_CONST(Action) { ITERATE_NODES_CONST(PROCESS_ALL_BODY(Action, _RO), true) }
 
-#define PROCESS_IN_BOX(Action) { ITERATE_NODES(PROCESS_IN_BOX_BODY(Action), NODE_IN_BOX) }
-#define PROCESS_IN_BOX_CONST(Action) { ITERATE_NODES_CONST(PROCESS_IN_BOX_BODY(Action), NODE_IN_BOX) }
+#define PROCESS_IN_BOX(Action) { ITERATE_NODES(PROCESS_IN_BOX_BODY(Action,), NODE_IN_BOX) }
+#define PROCESS_IN_BOX_CONST(Action) { ITERATE_NODES_CONST(PROCESS_IN_BOX_BODY(Action, _RO), NODE_IN_BOX) }
 
-#define PROCESS_IN_FRUSTUM(Action) { ITERATE_NODES(PROCESS_IN_FRUSTUM_BODY(Action), NODE_IN_FRUSTUM) }
-#define PROCESS_IN_FRUSTUM_CONST(Action) { ITERATE_NODES_CONST(PROCESS_IN_FRUSTUM_BODY(Action), NODE_IN_FRUSTUM) }
+#define PROCESS_IN_FRUSTUM(Action) { ITERATE_NODES(PROCESS_IN_FRUSTUM_BODY(Action,), NODE_IN_FRUSTUM) }
+#define PROCESS_IN_FRUSTUM_CONST(Action) { ITERATE_NODES_CONST(PROCESS_IN_FRUSTUM_BODY(Action, _RO), NODE_IN_FRUSTUM) }
 
 #define PROCESS_BY_RAY(Action)\
 {\
 	const float RadiusSq = Radius * Radius;\
-	ITERATE_NODES(PROCESS_BY_RAY_BODY(Action), false)\
+	ITERATE_NODES(PROCESS_BY_RAY_BODY(Action,), false)\
 }
 #define PROCESS_BY_RAY_CONST(Action)\
 {\
 	const float RadiusSq = Radius * Radius; \
-	ITERATE_NODES_CONST(PROCESS_BY_RAY_BODY(Action), false)\
+	ITERATE_NODES_CONST(PROCESS_BY_RAY_BODY(Action, _RO), false)\
 }
\ No newline at end of file
diff --git a/Engine/Plugins/Enterprise/LidarPointCloud/Source/LidarPointCloudRuntime/Private/LidarPointCloudSettings.cpp b/Engine/Plugins/Enterprise/LidarPointCloud/Source/LidarPointCloudRuntime/Private/LidarPointCloudSettings.cpp
--- a/Engine/Plugins/Enterprise/LidarPointCloud/Source/LidarPointCloudRuntime/Private/LidarPointCloudSettings.cpp
+++ b/Engine/Plugins/Enterprise/LidarPointCloud/Source/LidarPointCloudRuntime/Private/LidarPointCloudSettings.cpp
@@ -13,6 +13,8 @@ ULidarPointCloudSettings::ULidarPointCloudSettings()
 	, bUseAsyncImport(true)
 	, MaxImportBufferSize(52428800)
 	, ExportBatchSize(5000000)
+	, bPrioritizeActiveViewport(true)
+	, bUseCompression(true)
 	, MeshingBatchSize(128)
 	, bAutoCenterOnImport(true)
 	, ImportScale(100)
diff --git a/Engine/Plugins/Enterprise/LidarPointCloud/Source/LidarPointCloudRuntime/Private/LidarPointCloudShared.cpp b/Engine/Plugins/Enterprise/LidarPointCloud/Source/LidarPointCloudRuntime/Private/LidarPointCloudShared.cpp
--- a/Engine/Plugins/Enterprise/LidarPointCloud/Source/LidarPointCloudRuntime/Private/LidarPointCloudShared.cpp
+++ b/Engine/Plugins/Enterprise/LidarPointCloud/Source/LidarPointCloudRuntime/Private/LidarPointCloudShared.cpp
@@ -30,4 +30,123 @@ const FDoubleVector FDoubleVector::ZeroVector = FDoubleVector(FVector::ZeroVecto
 const FDoubleVector FDoubleVector::OneVector = FDoubleVector(FVector::OneVector);
 const FDoubleVector FDoubleVector::UpVector = FDoubleVector(FVector::UpVector);
 const FDoubleVector FDoubleVector::ForwardVector = FDoubleVector(FVector::ForwardVector);
-const FDoubleVector FDoubleVector::RightVector = FDoubleVector(FVector::RightVector);
\ No newline at end of file
+const FDoubleVector FDoubleVector::RightVector = FDoubleVector(FVector::RightVector);
+
+void FLidarPointCloudDataBuffer::MarkAsFree()
+{
+	if (PendingSize > 0)
+	{
+		Resize(PendingSize, true);
+		PendingSize = 0;
+	}
+
+	bInUse = false;
+}
+
+void FLidarPointCloudDataBuffer::Initialize(const int32& Size)
+{
+	Data.AddUninitialized(Size);
+}
+
+void FLidarPointCloudDataBuffer::Resize(const int32& NewBufferSize, bool bForce /*= false*/)
+{
+	if (bInUse && !bForce)
+	{
+		// Don't want to resize while in use, flag the new size as pending
+		// This will cause a resize as soon as the buffer is freed
+		PendingSize = NewBufferSize;
+	}
+	else
+	{
+		int32 Delta = NewBufferSize - Data.Num();
+
+		// Expand
+		if (Delta > 0)
+		{
+			Data.AddUninitialized(Delta);
+		}
+		// Shrink
+		else
+		{
+			Data.RemoveAtSwap(0, -Delta, true);
+		}
+	}
+}
+
+FLidarPointCloudDataBufferManager::FLidarPointCloudDataBufferManager(const int32& BufferSize)
+	: BufferSize(BufferSize)
+	, Head(FLidarPointCloudDataBuffer())
+	, Tail(&Head)
+{
+	Head.Element.Initialize(BufferSize);
+}
+
+FLidarPointCloudDataBufferManager::~FLidarPointCloudDataBufferManager()
+{
+	auto Iterator = &Head;
+	while (Iterator)
+	{
+		if (Iterator != &Head)
+		{
+			auto Tmp = Iterator;
+			Iterator = Iterator->Next;
+			delete Tmp;
+		}
+		else
+		{
+			Iterator = Iterator->Next;
+		}
+	}
+}
+
+FLidarPointCloudDataBuffer* FLidarPointCloudDataBufferManager::GetFreeBuffer()
+{
+	FLidarPointCloudDataBuffer* OutBuffer = nullptr;
+
+	// Find available memory allocation
+	{
+		auto Iterator = &Head;
+		while (Iterator)
+		{
+			if (!Iterator->Element.bInUse)
+			{
+				OutBuffer = &Iterator->Element;
+				break;
+			}
+
+			Iterator = Iterator->Next;
+		}
+	}
+
+	// If none found, add a new one
+	if (!OutBuffer)
+	{
+		Tail->Next = new TList<FLidarPointCloudDataBuffer>(FLidarPointCloudDataBuffer());
+		Tail = Tail->Next;
+		OutBuffer = &Tail->Element;
+		OutBuffer->Initialize(BufferSize);
+	}
+
+	OutBuffer->bInUse = true;
+
+	return OutBuffer;
+}
+
+void FLidarPointCloudDataBufferManager::Resize(const int32& NewBufferSize)
+{
+	// Skip, if no change required
+	if (BufferSize == NewBufferSize)
+	{
+		return;
+	}
+
+	BufferSize = NewBufferSize;
+
+	auto Iterator = &Head;
+	while (Iterator)
+	{
+		Iterator->Element.Resize(NewBufferSize);
+		Iterator = Iterator->Next;
+	}
+}
+
diff --git a/Engine/Plugins/Enterprise/LidarPointCloud/Source/LidarPointCloudRuntime/Private/Rendering/LidarPointCloudRenderBuffers.cpp b/Engine/Plugins/Enterprise/LidarPointCloud/Source/LidarPointCloudRuntime/Private/Rendering/LidarPointCloudRenderBuffers.cpp
--- a/Engine/Plugins/Enterprise/LidarPointCloud/Source/LidarPointCloudRuntime/Private/Rendering/LidarPointCloudRenderBuffers.cpp
+++ b/Engine/Plugins/Enterprise/LidarPointCloud/Source/LidarPointCloudRuntime/Private/Rendering/LidarPointCloudRenderBuffers.cpp
@@ -19,34 +19,33 @@
 
 //////////////////////////////////////////////////////////// Base Buffer
 
+TGlobalResource<FLidarPointCloudRenderBuffer> GLidarPointCloudRenderBuffer;
+TGlobalResource<FLidarPointCloudIndexBuffer> GLidarPointCloudIndexBuffer;
+TGlobalResource<FLidarPointCloudVertexFactory> GLidarPointCloudVertexFactory;
+
 void FLidarPointCloudBuffer::Resize(const uint32& RequestedCapacity)
 {
 	// This must be called from Rendering thread
 	check(IsInRenderingThread());
 
-	if (Capacity < RequestedCapacity)
+	if (Capacity != RequestedCapacity)
 	{
-		// Apply some slack to limit number of rebuilds
-		uint32 NewCapacity = Slack > 0 ? RequestedCapacity * (1 + Slack) : RequestedCapacity;
-
 		if (GetDefault<ULidarPointCloudSettings>()->bLogBufferExpansion)
 		{
-			PC_LOG("Resizing %s: %u => %u", *GetFriendlyName(), Capacity, NewCapacity);
+			PC_LOG("Resizing %s: %u => %u", *GetFriendlyName(), Capacity, RequestedCapacity);
 		}
 
 		Release();
-		Capacity = NewCapacity;
+		Capacity = RequestedCapacity;
 		Initialize();
 	}
 }
 
-//////////////////////////////////////////////////////////// Index Buffer
-
-FLidarPointCloudIndexBuffer::~FLidarPointCloudIndexBuffer()
+FLidarPointCloudBuffer::~FLidarPointCloudBuffer()
 {
 	FRenderCommandFence Fence;
 
-	ENQUEUE_RENDER_COMMAND(ReleasePointCloudIndexBuffer)(
+	ENQUEUE_RENDER_COMMAND(ReleaseLidarPointCloudRenderBuffer)(
 		[this](FRHICommandListImmediate& RHICmdList)
 		{
 			Release();
@@ -56,6 +55,8 @@ FLidarPointCloudIndexBuffer::~FLidarPointCloudIndexBuffer()
 	Fence.Wait();
 }
 
+//////////////////////////////////////////////////////////// Index Buffer
+
 void FLidarPointCloudIndexBuffer::InitRHI()
 {
 	FRHIResourceCreateInfo CreateInfo;
@@ -91,17 +92,17 @@ void FLidarPointCloudIndexBuffer::InitRHI()
 
 //////////////////////////////////////////////////////////// Structured Buffer
 
-void FLidarPointCloudStructuredBuffer::Initialize()
+void FLidarPointCloudRenderBuffer::InitRHI()
 {
 	// This must be called from Rendering thread
 	check(IsInRenderingThread());
 
 	FRHIResourceCreateInfo CreateInfo;
-	Buffer = RHICreateStructuredBuffer(ElementSize, ElementSize * Capacity, BUF_ShaderResource | BUF_Dynamic, CreateInfo);
-	SRV = RHICreateShaderResourceView(Buffer);
+	Buffer = RHICreateVertexBuffer(ElementSize * Capacity, BUF_ShaderResource | BUF_Dynamic, CreateInfo);
+	SRV = RHICreateShaderResourceView(Buffer, ElementSize, PF_R32_FLOAT);
 }
 
-void FLidarPointCloudStructuredBuffer::Release()
+void FLidarPointCloudRenderBuffer::ReleaseRHI()
 {
 	// This must be called from Rendering thread
 	check(IsInRenderingThread());
@@ -115,234 +116,6 @@ void FLidarPointCloudStructuredBuffer::Release()
 	SRV.SafeRelease();
 }
 
-//////////////////////////////////////////////////////////// Instance Buffer
-
-void FLidarPointCloudInstanceBuffer::Reset()
-{
-	NewNumInstances = 0;
-	Nodes.Reset();
-}
-
-uint32 FLidarPointCloudInstanceBuffer::UpdateData(bool bUseClassification)
-{
-	// Populate buffer only if there is any data to populate with
-	if (NewNumInstances)
-	{
-		// Rebuild resource if the new data size is greater than the current buffer can accommodate
-		// 17 bytes per point, element size set to 4 bytes - estimated 4.3 elements per point
-		StructuredBuffer.Resize(NewNumInstances * 4.3);
-
-		FColor SelectionColor = FColor::White;
-
-#if WITH_EDITOR
-		SelectionColor = GetDefault<UEditorStyleSettings>()->SelectionColor.ToFColor(false);
-#endif
-
-		// Send new data
-		uint8* BufferStart = (uint8*)RHILockStructuredBuffer(StructuredBuffer.Buffer, 0, StructuredBuffer.GetSize(), RLM_WriteOnly);
-		uint8* BufferCurrent = BufferStart;
-		const FLidarPointCloudPoint* DataPointer = nullptr;
-
-		if (bOwnedByEditor)
-		{
-			if (bUseClassification)
-			{
-				for (auto Node : Nodes)
-				{
-					const auto DataNode = Node->DataNode;
-					for (FLidarPointCloudPoint* Data = DataNode->AllocatedPoints.GetData(), *DataEnd = Data + DataNode->AllocatedPoints.Num(); Data != DataEnd; ++Data)
-					{
-						if (!Data->bVisible)
-						{
-							break;
-						}
-
-						FMemory::Memcpy(BufferCurrent, Data, 12);
-						BufferCurrent += 12;
-
-						if (Data->bSelected)
-						{
-							FMemory::Memcpy(BufferCurrent, &SelectionColor, 4);
-						}
-						else
-						{
-							FColor ClassificationColor(Data->ClassificationID, Data->ClassificationID, Data->ClassificationID, Data->Color.A);
-							FMemory::Memcpy(BufferCurrent, &ClassificationColor, 4);
-						}
-
-						BufferCurrent += 4;
-					}
-
-					for (FLidarPointCloudPoint* Data = DataNode->PaddingPoints.GetData(), *DataEnd = Data + DataNode->PaddingPoints.Num(); Data != DataEnd; ++Data)
-					{
-						if (!Data->bVisible)
-						{
-							break;
-						}
-
-						FMemory::Memcpy(BufferCurrent, Data, 12);
-						BufferCurrent += 12;
-
-						if (Data->bSelected)
-						{
-							FMemory::Memcpy(BufferCurrent, &SelectionColor, 4);
-						}
-						else
-						{
-							FColor ClassificationColor(Data->ClassificationID, Data->ClassificationID, Data->ClassificationID, Data->Color.A);
-							FMemory::Memcpy(BufferCurrent, &ClassificationColor, 4);
-						}
-
-						BufferCurrent += 4;
-					}
-				}
-			}
-			else
-			{
-				for (auto Node : Nodes)
-				{
-					const auto DataNode = Node->DataNode;
-					for (FLidarPointCloudPoint* Data = DataNode->AllocatedPoints.GetData(), *DataEnd = Data + DataNode->AllocatedPoints.Num(); Data != DataEnd; ++Data)
-					{
-						if (!Data->bVisible)
-						{
-							break;
-						}
-
-						if (Data->bSelected)
-						{
-							FMemory::Memcpy(BufferCurrent, Data, 12);
-							BufferCurrent += 12;
-
-							FMemory::Memcpy(BufferCurrent, &SelectionColor, 4);
-							BufferCurrent += 4;
-						}
-						else
-						{
-							FMemory::Memcpy(BufferCurrent, Data, 16);
-							BufferCurrent += 16;
-						}
-					}
-
-					for (FLidarPointCloudPoint* Data = DataNode->PaddingPoints.GetData(), *DataEnd = Data + DataNode->PaddingPoints.Num(); Data != DataEnd; ++Data)
-					{
-						if (!Data->bVisible)
-						{
-							break;
-						}
-
-						if (Data->bSelected)
-						{
-							FMemory::Memcpy(BufferCurrent, Data, 12);
-							BufferCurrent += 12;
-
-							FMemory::Memcpy(BufferCurrent, &SelectionColor, 4);
-							BufferCurrent += 4;
-						}
-						else
-						{
-							FMemory::Memcpy(BufferCurrent, Data, 16);
-							BufferCurrent += 16;
-						}
-					}
-				}
-			}
-		}
-		else
-		{
-			if (bUseClassification)
-			{
-				for (auto Node : Nodes)
-				{
-					const auto DataNode = Node->DataNode;
-					for (FLidarPointCloudPoint* Data = DataNode->AllocatedPoints.GetData(), *DataEnd = Data + DataNode->AllocatedPoints.Num(); Data != DataEnd; ++Data)
-					{
-						if (!Data->bVisible)
-						{
-							break;
-						}
-
-						FMemory::Memcpy(BufferCurrent, Data, 12);
-						BufferCurrent += 12;
-
-						FColor ClassificationColor(Data->ClassificationID, Data->ClassificationID, Data->ClassificationID, Data->Color.A);
-						FMemory::Memcpy(BufferCurrent, &ClassificationColor, 4);
-						BufferCurrent += 4;
-					}
-
-					for (FLidarPointCloudPoint* Data = DataNode->PaddingPoints.GetData(), *DataEnd = Data + DataNode->PaddingPoints.Num(); Data != DataEnd; ++Data)
-					{
-						if (!Data->bVisible)
-						{
-							break;
-						}
-
-						FMemory::Memcpy(BufferCurrent, Data, 12);
-						BufferCurrent += 12;
-
-						FColor ClassificationColor(Data->ClassificationID, Data->ClassificationID, Data->ClassificationID, Data->Color.A);
-						FMemory::Memcpy(BufferCurrent, &ClassificationColor, 4);
-						BufferCurrent += 4;
-					}
-				}
-			}
-			else
-			{
-				for (auto Node : Nodes)
-				{
-					const auto DataNode = Node->DataNode;
-					for (FLidarPointCloudPoint* Data = DataNode->AllocatedPoints.GetData(), *DataEnd = Data + DataNode->AllocatedPoints.Num(); Data != DataEnd; ++Data)
-					{
-						if (!Data->bVisible)
-						{
-							break;
-						}
-
-						FMemory::Memcpy(BufferCurrent, Data, 16);
-						BufferCurrent += 16;
-					}
-
-					for (FLidarPointCloudPoint* Data = DataNode->PaddingPoints.GetData(), *DataEnd = Data + DataNode->PaddingPoints.Num(); Data != DataEnd; ++Data)
-					{
-						if (!Data->bVisible)
-						{
-							break;
-						}
-
-						FMemory::Memcpy(BufferCurrent, Data, 16);
-						BufferCurrent += 16;
-					}
-				}
-			}
-		}
-
-		// Calculates the actual number of instances copied to the buffer (accounts for the invisible points)
-		NewNumInstances = (BufferCurrent - BufferStart) / 16;
-
-		// Send Scale data
-		for (auto Node : Nodes)
-		{
-			FMemory::Memset(BufferCurrent, Node->VirtualDepth, Node->DataNode->GetNumVisiblePoints());
-			BufferCurrent += Node->DataNode->GetNumVisiblePoints();
-		}
-		RHIUnlockStructuredBuffer(StructuredBuffer.Buffer);
-	}
-
-	return NumInstances = NewNumInstances;
-}
-
-void FLidarPointCloudInstanceBuffer::AddNode(const FLidarPointCloudTraversalOctreeNode* Node)
-{
-	Nodes.Add(Node);
-	NewNumInstances += Node->DataNode->GetNumVisiblePoints();
-}
-
-void FLidarPointCloudInstanceBuffer::Release()
-{
-	Nodes.Empty();
-	NumInstances = 0;
-}
-
 //////////////////////////////////////////////////////////// User Data
 
 void FLidarPointCloudBatchElementUserData::SetClassificationColors(const TMap<int32, FLinearColor>& InClassificationColors)
@@ -360,6 +133,8 @@ void FLidarPointCloudVertexFactoryShaderParameters::Bind(const FShaderParameterM
 {
 	BINDPARAM(DataBuffer);
 	BINDPARAM(IndexDivisor);
+	BINDPARAM(FirstElementIndex);
+	BINDPARAM(LocationOffset);
 	BINDPARAM(VDMultiplier);
 	BINDPARAM(SizeOffset);
 	BINDPARAM(RootCellSize);
@@ -368,7 +143,6 @@ void FLidarPointCloudVertexFactoryShaderParameters::Bind(const FShaderParameterM
 	BINDPARAM(ViewRightVector);
 	BINDPARAM(ViewUpVector);
 	BINDPARAM(BoundsSize);
-	BINDPARAM(BoundsOffset);
 	BINDPARAM(ElevationColorBottom);
 	BINDPARAM(ElevationColorTop);
 	BINDPARAM(bUseCircle);
@@ -391,6 +165,8 @@ void FLidarPointCloudVertexFactoryShaderParameters::GetElementShaderBindings(con
 
 	SETSRVPARAM(DataBuffer);
 	SETPARAM(IndexDivisor);
+	SETPARAM(FirstElementIndex);
+	SETPARAM(LocationOffset);
 	SETPARAM(VDMultiplier);
 	SETPARAM(SizeOffset);
 	SETPARAM(RootCellSize);
@@ -399,7 +175,6 @@ void FLidarPointCloudVertexFactoryShaderParameters::GetElementShaderBindings(con
 	SETPARAM(ViewRightVector);
 	SETPARAM(ViewUpVector);
 	SETPARAM(BoundsSize);
-	SETPARAM(BoundsOffset);
 	SETPARAM(ElevationColorBottom);
 	SETPARAM(ElevationColorTop);
 	SETPARAM(bUseCircle);
@@ -415,25 +190,9 @@ void FLidarPointCloudVertexFactoryShaderParameters::GetElementShaderBindings(con
 	SETPARAM(ClassificationColors);
 }
 
-FLidarPointCloudVertexFactory::~FLidarPointCloudVertexFactory()
-{
-	FRenderCommandFence Fence;
-
-	FVertexFactory* VertexFactory = this;
-
-	ENQUEUE_RENDER_COMMAND(ReleasePointCloudIndexBuffer)(
-		[VertexFactory](FRHICommandListImmediate& RHICmdList)
-		{
-			VertexFactory->ReleaseResource();
-		});
-
-	Fence.BeginFence();
-	Fence.Wait();
-}
-
 bool FLidarPointCloudVertexFactory::ShouldCompilePermutation(const FVertexFactoryShaderPermutationParameters& Parameters)
 {
-	return (IsPCPlatform(Parameters.Platform) && IsFeatureLevelSupported(Parameters.Platform, ERHIFeatureLevel::SM5) && 
+	return (IsPCPlatform(Parameters.Platform) && IsFeatureLevelSupported(Parameters.Platform, ERHIFeatureLevel::SM5) &&
 		Parameters.MaterialParameters.MaterialDomain == MD_Surface && Parameters.MaterialParameters.bIsUsedWithLidarPointCloud) || Parameters.MaterialParameters.bIsSpecialEngineMaterial;
 }
 
