Merging

diff --git a/Engine/Plugins/Enterprise/LidarPointCloud/Source/LidarPointCloudRuntime/Private/Rendering/LidarPointCloudRenderBuffers.h b/Engine/Plugins/Enterprise/LidarPointCloud/Source/LidarPointCloudRuntime/Private/Rendering/LidarPointCloudRenderBuffers.h
--- a/Engine/Plugins/Enterprise/LidarPointCloud/Source/LidarPointCloudRuntime/Private/Rendering/LidarPointCloudRenderBuffers.h
+++ b/Engine/Plugins/Enterprise/LidarPointCloud/Source/LidarPointCloudRuntime/Private/Rendering/LidarPointCloudRenderBuffers.h
@@ -7,7 +7,7 @@
 #include "VertexFactory.h"
 #include "LidarPointCloudShared.h"
 #include "HAL/ThreadSafeBool.h"
-#include "GlobalShader.h"
+#include "MeshMaterialShader.h"
 
 /**
  * Base class for the buffers
@@ -18,7 +18,7 @@ public:
 	/** Resizes the buffer to the specified capacity, if necessary. Must be called from Rendering thread. */
 	void Resize(const uint32& RequestedCapacity);
 
-	FLidarPointCloudBuffer(const uint32& ElementSize, const float& Slack, const uint32 InitialCapacity = 0) : Capacity(InitialCapacity), ElementSize(ElementSize), Slack(Slack) {}
+	FLidarPointCloudBuffer(const uint32& ElementSize) : Capacity(100000), ElementSize(ElementSize) {}
 
 	virtual FString GetFriendlyName() const { return TEXT("FPointCloudBuffer"); }
 
@@ -28,9 +28,9 @@ public:
 protected:
 	uint32 Capacity;
 	uint32 ElementSize;
-	float Slack;
 
-	FLidarPointCloudBuffer() : FLidarPointCloudBuffer(0, 0) {}
+	FLidarPointCloudBuffer() : FLidarPointCloudBuffer(0) {}
+	virtual ~FLidarPointCloudBuffer();
 
 	virtual void Initialize() = 0;
 	virtual void Release() = 0;
@@ -39,11 +39,10 @@ protected:
 /**
  * This class creates an IndexBuffer shared between all assets and all instances.
  */
-class FLidarPointCloudIndexBuffer : public FIndexBuffer, public FLidarPointCloudBuffer
+class FLidarPointCloudIndexBuffer : public FLidarPointCloudBuffer, public FIndexBuffer
 {
 public:
-	FLidarPointCloudIndexBuffer() : FLidarPointCloudBuffer(sizeof(uint32), 0, 1000000) { }
-	~FLidarPointCloudIndexBuffer();
+	FLidarPointCloudIndexBuffer() : FLidarPointCloudBuffer(sizeof(uint32)) { }
 
 	uint32 GetPointOffset() const { return PointOffset; }
 
@@ -58,68 +57,31 @@ private:
 };
 
 /**
- * Encapsulates a GPU read structured buffer with its SRV.
+ * Encapsulates a GPU read buffer with its SRV.
  */
-class FLidarPointCloudStructuredBuffer : public FLidarPointCloudBuffer
+class FLidarPointCloudRenderBuffer : public FLidarPointCloudBuffer, public FRenderResource
 {
 public:
-	FStructuredBufferRHIRef Buffer;
+	FVertexBufferRHIRef Buffer;
 	FShaderResourceViewRHIRef SRV;
 
-	FLidarPointCloudStructuredBuffer(const uint32& ElementSize, const float& Slack) : FLidarPointCloudBuffer(ElementSize, Slack) {}
-	virtual ~FLidarPointCloudStructuredBuffer() { Release(); }
+	int32 PointCount;
 
-	virtual void Initialize() override;
-	virtual void Release() override;
-	virtual FString GetFriendlyName() const override { return TEXT("Structured Buffer"); }
+	FLidarPointCloudRenderBuffer() : FLidarPointCloudBuffer(4) {}
+	virtual ~FLidarPointCloudRenderBuffer() { Release(); }
 
-private:
-	FLidarPointCloudStructuredBuffer(const FLidarPointCloudStructuredBuffer&) = delete;
-	FLidarPointCloudStructuredBuffer(FLidarPointCloudStructuredBuffer&&) = delete;
-	FLidarPointCloudStructuredBuffer& operator=(const FLidarPointCloudStructuredBuffer&) = delete;
-	FLidarPointCloudStructuredBuffer& operator=(FLidarPointCloudStructuredBuffer&&) = delete;
-};
+	virtual void InitRHI() override;
+	virtual void ReleaseRHI() override;
 
-/**
- * Handles creation of the dynamic Instance Data
- */
- // #refactor: Can we avoid referencing FLidarPointCloudOctreeNode here?
-struct FLidarPointCloudTraversalOctreeNode;
-class FLidarPointCloudInstanceBuffer
-{
-public:
-	FLidarPointCloudStructuredBuffer StructuredBuffer;
-	TArray<const FLidarPointCloudTraversalOctreeNode*> Nodes;
-	bool bOwnedByEditor;
+	virtual void Initialize() override { InitResource(); }
+	virtual void Release() override { ReleaseResource(); }
+	virtual FString GetFriendlyName() const override { return TEXT("FLidarPointCloudRenderBuffer"); }
 
 private:
-	uint32 NumInstances;
-	uint32 NewNumInstances;
-
-public:
-	FLidarPointCloudInstanceBuffer()
-		: StructuredBuffer(4, 1.0f)
-		, bOwnedByEditor(false)
-		, NumInstances(0)
-		, NewNumInstances(0)
-	{
-	}
-
-	FORCEINLINE uint32 GetNumInstances() const { return NumInstances; }
-
-	/** Resets the nodes assigned for building this buffer */
-	void Reset();
-
-	/**
-	 * Resizes and updates the buffer on the GPU if necessary.
-	 * Returns the number of instances used.
-	 */
-	uint32 UpdateData(bool bUseClassification);
-
-	void AddNode(const FLidarPointCloudTraversalOctreeNode* Node);
-
-	/** Releases the memory held by StructuredBuffer */
-	void Release();
+	FLidarPointCloudRenderBuffer(const FLidarPointCloudRenderBuffer&) = delete;
+	FLidarPointCloudRenderBuffer(FLidarPointCloudRenderBuffer&&) = delete;
+	FLidarPointCloudRenderBuffer& operator=(const FLidarPointCloudRenderBuffer&) = delete;
+	FLidarPointCloudRenderBuffer& operator=(FLidarPointCloudRenderBuffer&&) = delete;
 };
 
 /**
@@ -129,6 +91,8 @@ struct FLidarPointCloudBatchElementUserData
 {
 	FRHIShaderResourceView* DataBuffer;
 	int32 IndexDivisor;
+	int32 FirstElementIndex;
+	FVector LocationOffset;
 	float VDMultiplier;
 	int32 SizeOffset;
 	float RootCellSize;
@@ -137,7 +101,6 @@ struct FLidarPointCloudBatchElementUserData
 	FVector ViewRightVector;
 	FVector ViewUpVector;
 	FVector BoundsSize;
-	FVector BoundsOffset;
 	FVector ElevationColorBottom;
 	FVector ElevationColorTop;
 	int32 bUseCircle;
@@ -169,36 +132,38 @@ struct FLidarPointCloudBatchElementUserData
  */
 class FLidarPointCloudVertexFactoryShaderParameters : public FVertexFactoryShaderParameters
 {
+	DECLARE_INLINE_TYPE_LAYOUT(FLidarPointCloudVertexFactoryShaderParameters, NonVirtual);
+
 public:
 	void Bind(const FShaderParameterMap& ParameterMap);
 	void GetElementShaderBindings(const class FSceneInterface* Scene, const FSceneView* View, const FMeshMaterialShader* Shader, const EVertexInputStreamType InputStreamType, ERHIFeatureLevel::Type FeatureLevel,
 		const FVertexFactory* VertexFactory, const FMeshBatchElement& BatchElement, class FMeshDrawSingleShaderBindings& ShaderBindings, FVertexInputStreamArray& VertexStreams) const;
 
-private:
-	FShaderResourceParameter DataBuffer;
-	FShaderParameter IndexDivisor;
-	FShaderParameter VDMultiplier;
-	FShaderParameter SizeOffset;
-	FShaderParameter RootCellSize;
-	FShaderParameter bUseLODColoration;
-	FShaderParameter SpriteSizeMultiplier;
-	FShaderParameter ViewRightVector;
-	FShaderParameter ViewUpVector;
-	FShaderParameter BoundsSize;
-	FShaderParameter BoundsOffset;
-	FShaderParameter ElevationColorBottom;
-	FShaderParameter ElevationColorTop;
-	FShaderParameter bUseCircle;
-	FShaderParameter bUseColorOverride;
-	FShaderParameter bUseElevationColor;
-	FShaderParameter Offset;
-	FShaderParameter Contrast;
-	FShaderParameter Saturation;
-	FShaderParameter Gamma;
-	FShaderParameter Tint;
-	FShaderParameter IntensityInfluence;
-	FShaderParameter bUseClassification;
-	FShaderParameter ClassificationColors;
+	LAYOUT_FIELD(FShaderResourceParameter, DataBuffer);
+	LAYOUT_FIELD(FShaderParameter, IndexDivisor);
+	LAYOUT_FIELD(FShaderParameter, FirstElementIndex);
+	LAYOUT_FIELD(FShaderParameter, LocationOffset);
+	LAYOUT_FIELD(FShaderParameter, VDMultiplier);
+	LAYOUT_FIELD(FShaderParameter, SizeOffset);
+	LAYOUT_FIELD(FShaderParameter, RootCellSize);
+	LAYOUT_FIELD(FShaderParameter, bUseLODColoration);
+	LAYOUT_FIELD(FShaderParameter, SpriteSizeMultiplier);
+	LAYOUT_FIELD(FShaderParameter, ViewRightVector);
+	LAYOUT_FIELD(FShaderParameter, ViewUpVector);
+	LAYOUT_FIELD(FShaderParameter, BoundsSize);
+	LAYOUT_FIELD(FShaderParameter, ElevationColorBottom);
+	LAYOUT_FIELD(FShaderParameter, ElevationColorTop);
+	LAYOUT_FIELD(FShaderParameter, bUseCircle);
+	LAYOUT_FIELD(FShaderParameter, bUseColorOverride);
+	LAYOUT_FIELD(FShaderParameter, bUseElevationColor);
+	LAYOUT_FIELD(FShaderParameter, Offset);
+	LAYOUT_FIELD(FShaderParameter, Contrast);
+	LAYOUT_FIELD(FShaderParameter, Saturation);
+	LAYOUT_FIELD(FShaderParameter, Gamma);
+	LAYOUT_FIELD(FShaderParameter, Tint);
+	LAYOUT_FIELD(FShaderParameter, IntensityInfluence);
+	LAYOUT_FIELD(FShaderParameter, bUseClassification);
+	LAYOUT_FIELD(FShaderParameter, ClassificationColors);
 };
 
 /**
@@ -209,11 +174,10 @@ class FLidarPointCloudVertexFactory : public FVertexFactory
 	DECLARE_VERTEX_FACTORY_TYPE(FLidarPointCloudVertexFactory);
 
 public:
-	static bool ShouldCompilePermutation(const FVertexFactoryShaderPermutationParameters& Parameters);
-	static void ModifyCompilationEnvironment(const FVertexFactoryShaderPermutationParameters& Parameters, FShaderCompilerEnvironment& OutEnvironment) {}
+	static bool ShouldCache(const FVertexFactoryShaderPermutationParameters& Parameters) { return IsFeatureLevelSupported(Parameters.Platform, ERHIFeatureLevel::SM5); }
+    static bool ShouldCompilePermutation(const FVertexFactoryShaderPermutationParameters& Parameters);
 
 	FLidarPointCloudVertexFactory() : FVertexFactory(ERHIFeatureLevel::SM5) { }
-	~FLidarPointCloudVertexFactory();
 
 private:
 	/** Very simple implementation of a ZeroStride Vertex Buffer */
@@ -235,4 +199,9 @@ private:
 
 	virtual void InitRHI() override;
 	virtual void ReleaseRHI() override;
-};
\ No newline at end of file
+};
+
+/** A set of global render resources shared between all Lidar Point Cloud proxies */
+extern TGlobalResource<FLidarPointCloudRenderBuffer> GLidarPointCloudRenderBuffer;
+extern TGlobalResource<FLidarPointCloudIndexBuffer> GLidarPointCloudIndexBuffer;
+extern TGlobalResource<FLidarPointCloudVertexFactory> GLidarPointCloudVertexFactory;
\ No newline at end of file
diff --git a/Engine/Plugins/Enterprise/LidarPointCloud/Source/LidarPointCloudRuntime/Private/Rendering/LidarPointCloudRendering.cpp b/Engine/Plugins/Enterprise/LidarPointCloud/Source/LidarPointCloudRuntime/Private/Rendering/LidarPointCloudRendering.cpp
--- a/Engine/Plugins/Enterprise/LidarPointCloud/Source/LidarPointCloudRuntime/Private/Rendering/LidarPointCloudRendering.cpp
+++ b/Engine/Plugins/Enterprise/LidarPointCloud/Source/LidarPointCloudRuntime/Private/Rendering/LidarPointCloudRendering.cpp
@@ -1,10 +1,12 @@
 // Copyright Epic Games, Inc. All Rights Reserved.
 
+#include "Rendering/LidarPointCloudRendering.h"
 #include "LidarPointCloudRenderBuffers.h"
 #include "LidarPointCloudComponent.h"
 #include "LidarPointCloud.h"
 #include "LidarPointCloudShared.h"
 #include "LidarPointCloudOctree.h"
+#include "LidarPointCloudLODManager.h"
 #include "PrimitiveSceneProxy.h"
 #include "MeshBatch.h"
 #include "Engine/Engine.h"
@@ -13,17 +15,15 @@
 #include "LocalVertexFactory.h"
 #include "Materials/Material.h"
 
-DECLARE_DWORD_COUNTER_STAT(TEXT("Visible Points"), STAT_PointCount, STATGROUP_LidarPointCloud)
-DECLARE_DWORD_COUNTER_STAT(TEXT("Visible Nodes"), STAT_NodeCount, STATGROUP_LidarPointCloud)
 DECLARE_DWORD_COUNTER_STAT(TEXT("Draw Calls"), STAT_DrawCallCount, STATGROUP_LidarPointCloud)
 
-DECLARE_CYCLE_STAT(TEXT("Buffer Update"), STAT_BufferUpdate, STATGROUP_LidarPointCloud);
-
-/** A global instance of an Index Buffer shared between all Lidar Point Cloud proxies */
-TGlobalResource<FLidarPointCloudIndexBuffer> GLidarPointCloudIndexBuffer;
-
-/** A global instance of a Vertex Factory shared between all Lidar Point Cloud proxies */
-TGlobalResource<FLidarPointCloudVertexFactory> GLidarPointCloudVertexFactory;
+FLidarPointCloudProxyUpdateData::FLidarPointCloudProxyUpdateData()
+	: FirstElementIndex(0)
+	, NumElements(0)
+	, VDMultiplier(1)
+	, RootCellSize(1)
+{
+}
 
 class FLidarPointCloudCollisionRendering
 {
@@ -89,9 +89,9 @@ private:
 			if (RHISupportsManualVertexFetch(GMaxRHIShaderPlatform))
 			{
 				NewData.PositionComponentSRV = RHICreateShaderResourceView(InVertexBuffer->VertexBufferRHI, 4, PF_R32_FLOAT);
-				NewData.ColorComponentsSRV = RHICreateShaderResourceView(nullptr, 4, PF_R8G8B8A8);
-				NewData.TangentsSRV = RHICreateShaderResourceView(nullptr, 4, PF_R8G8B8A8);
-				NewData.TextureCoordinatesSRV = RHICreateShaderResourceView(nullptr, 4, PF_G16R16F);
+				NewData.ColorComponentsSRV = RHICreateShaderResourceView(InVertexBuffer->VertexBufferRHI, 4, PF_R8G8B8A8);
+				NewData.TangentsSRV = RHICreateShaderResourceView(InVertexBuffer->VertexBufferRHI, 4, PF_R8G8B8A8);
+				NewData.TextureCoordinatesSRV = RHICreateShaderResourceView(InVertexBuffer->VertexBufferRHI, 4, PF_G16R16F);
 			}
 
 			Data = NewData;
@@ -153,14 +153,13 @@ private:
 	int32 MaxVertexIndex;
 };
 
-class FLidarPointCloudSceneProxy : public FPrimitiveSceneProxy
+class FLidarPointCloudSceneProxy : public ILidarPointCloudSceneProxy, public FPrimitiveSceneProxy
 {
 public:
 	FLidarPointCloudSceneProxy(ULidarPointCloudComponent* Component)
 		: FPrimitiveSceneProxy(Component)
-		, bReadyToRender(false)
+		, ProxyWrapper(MakeShared<FLidarPointCloudSceneProxyWrapper, ESPMode::ThreadSafe>(this))
 		, Component(Component)
-		, Octree(&Component->GetPointCloud()->Octree)
 		, Owner(Component->GetOwner())
 		, CollisionRendering(Component->GetPointCloud()->CollisionRendering)
 	{
@@ -174,7 +173,7 @@ public:
 	{
 		QUICK_SCOPE_CYCLE_COUNTER(STAT_PointCloudSceneProxy_GetDynamicMeshElements);
 
-		if (!bReadyToRender)
+		if (!Component->GetMaterial(0))
 		{
 			return;
 		}
@@ -188,8 +187,7 @@ public:
 			if (IsShown(View) && (VisibilityMap & (1 << ViewIndex)))
 			{
 				// Prepare the draw call
-				int32 NumElements = InstanceBuffer.GetNumInstances();
-				if (NumElements)
+				if (RenderData.NumElements)
 				{
 					FMeshBatch& MeshBatch = Collector.AllocateMesh();
 
@@ -206,7 +204,7 @@ public:
 					BatchElement.IndexBuffer = &GLidarPointCloudIndexBuffer;
 					BatchElement.FirstIndex = bUsesSprites ? 0 : GLidarPointCloudIndexBuffer.GetPointOffset();
 					BatchElement.MinVertexIndex = 0;
-					BatchElement.NumPrimitives = NumElements * (bUsesSprites ? 2 : 1);
+					BatchElement.NumPrimitives = RenderData.NumElements * (bUsesSprites ? 2 : 1);
 					BatchElement.UserData = &UserData[UserData.Add(BuildUserDataElement(View))];
 
 					Collector.AddMesh(ViewIndex, MeshBatch);
@@ -220,10 +218,9 @@ public:
 				// Draw selected nodes' bounds
 				if (Component->bDrawNodeBounds)
 				{
-					for (auto Node : SelectedNodes)
+					for (auto& Node : RenderData.Bounds)
 					{
-						FVector Extent = TraversalOctree->Extents[Node->Depth];
-						DrawWireBox(PDI, FBox(Node->Center - Extent, Node->Center + Extent), FColor(72, 72, 255), SDPG_World);
+						DrawWireBox(PDI, Node, FColor(72, 72, 255), SDPG_World);
 					}
 				}
 
@@ -269,73 +266,14 @@ public:
 		return Result;
 	}
 
-	/** Calculates node visibility and updates InstanceBuffer */
 	virtual void* InitViewCustomData(const FSceneView& InView, float InViewLODScale, FMemStackBase& InCustomDataMemStack, bool InIsStaticRelevant, bool InIsShadowOnly, const FLODMask* InVisiblePrimitiveLODMask = nullptr, float InMeshScreenSizeSquared = -1.0f) override
 	{
 		// Don't process for shadow views
-		if (InIsShadowOnly)
-		{
-			return nullptr;
-		}
-
-		bReadyToRender = false;
-
-		// Check if the data is available
-		if (!Octree || Octree->GetNumNodes() == 0 || !TraversalOctree.IsValid() || TraversalOctree->NumLODs == 0 || !Component->GetMaterial(0))
-		{
-			return nullptr;
-		}
-
-		// Check if the access to the data has been restricted
-		FScopeTryLock Lock(&Octree->DataLock);
-		if (!Lock.IsLocked())
-		{
-			return nullptr;
-		}
-
-		// Re-initialize the traversal octree, if needed
-		if (!TraversalOctree->bValid)
-		{
-			TraversalOctree.Reset(new FLidarPointCloudTraversalOctree(Octree, Component->GetComponentTransform()));
-		}
-
-		// Clamp the budget to 25M
-		const int32 PointBudget = FMath::Clamp(Component->PointBudget, 0, 25000000);
-
-		// Expand the Index Buffer, if necessary
-		GLidarPointCloudIndexBuffer.Resize(PointBudget);
-
-		// Prepare the buffer
-		InstanceBuffer.Reset();
-
-		// Select the relevant nodes to be rendered
-		SelectedNodes = TraversalOctree->GetVisibleNodes(&InView, FMath::Max(Component->MinScreenSize * Component->MinScreenSize, 0.0f), PointBudget, Component->ScreenCenterImportance, Component->MinDepth, Component->MaxDepth);
-		INC_DWORD_STAT_BY(STAT_NodeCount, SelectedNodes.Num());
-
-		// Calculate sprite size, if needed
-		if (Component->PointSize > 0)
-		{
-			TraversalOctree->CalculateSpriteSize(SelectedNodes, Component->PointSizeBias);
-		}
-
-		// Allocate the selected nodes
-		for (auto Node : SelectedNodes)
-		{
-			InstanceBuffer.AddNode(Node);
-		}
-
-		// Update buffers
+		if (!InIsShadowOnly)
 		{
-			SCOPE_CYCLE_COUNTER(STAT_BufferUpdate);
-			InstanceBuffer.UpdateData(Component->ColorSource == ELidarPointCloudColorationMode::Classification);
+			UserData.Reset(10);
 		}
 
-		bReadyToRender = true;
-
-		UserData.Reset(10);
-
-		INC_DWORD_STAT_BY(STAT_PointCount, InstanceBuffer.GetNumInstances());
-
 		return nullptr;
 	}
 
@@ -361,15 +299,12 @@ public:
 
 	/** UserData is used to pass rendering information to the VertexFactory */
 	FLidarPointCloudBatchElementUserData BuildUserDataElement(const FSceneView* InView) const
-	{
-		// Each point has a virtual LOD in range of 0 to Octree->GetNumLODs(), mapped to 0 - 255 range
-		// To reverse the mapping: x / 255.0f * Octree->GetNumLODs() or x * (1.0f / 255.0f * Octree->GetNumLODs())
-		constexpr float InvertedMultiplier = 1.0f / 255.0f;
-		FLidarPointCloudBatchElementUserData UserDataElement(InvertedMultiplier * Octree->GetNumLODs(), Octree->GetRootCellSize());
+	{	
+		FLidarPointCloudBatchElementUserData UserDataElement(RenderData.VDMultiplier, RenderData.RootCellSize);
 
 		const bool bUsesSprites = Component->PointSize > 0;
 
-		auto BoundsSize = Component->GetPointCloud()->GetPointsBounds().GetSize();
+		auto BoundsSize = Component->GetPointCloud()->GetBounds().GetSize();
 
 		// Make sure to apply minimum bounds size
 		BoundsSize.X = FMath::Max(BoundsSize.X, 0.001f);
@@ -378,14 +313,16 @@ public:
 
 		// Update shader parameters
 		UserDataElement.IndexDivisor = bUsesSprites ? 4 : 1;
+		UserDataElement.FirstElementIndex = RenderData.FirstElementIndex;
 
-		UserDataElement.SizeOffset = InstanceBuffer.GetNumInstances() * 4;
+		UserDataElement.LocationOffset = Component->GetPointCloud()->GetLocationOffset().ToVector();
+		
+		UserDataElement.SizeOffset = GLidarPointCloudRenderBuffer.PointCount * 4;
 		UserDataElement.SpriteSizeMultiplier = Component->PointSize * Component->GetComponentScale().GetAbsMax();
 
 		UserDataElement.ViewRightVector = InView->GetViewRight();
 		UserDataElement.ViewUpVector = InView->GetViewUp();
 		UserDataElement.BoundsSize = BoundsSize;
-		UserDataElement.BoundsOffset = Component->GetPointCloud()->GetBounds().GetCenter() - Component->GetPointCloud()->GetPointsBounds().GetCenter();
 		UserDataElement.ElevationColorBottom = FVector(Component->ColorSource == ELidarPointCloudColorationMode::None ? FColor::White : Component->ElevationColorBottom);
 		UserDataElement.ElevationColorTop = FVector(Component->ColorSource == ELidarPointCloudColorationMode::None ? FColor::White : Component->ElevationColorTop);
 		UserDataElement.bUseCircle = bUsesSprites && Component->PointShape == ELidarPointCloudSpriteShape::Circle;
@@ -402,7 +339,7 @@ public:
 		UserDataElement.SetClassificationColors(Component->ClassificationColors);
 
 		// Make sure to update the reference to the resource, in case it was re-initialized
-		UserDataElement.DataBuffer = InstanceBuffer.StructuredBuffer.SRV;
+		UserDataElement.DataBuffer = GLidarPointCloudRenderBuffer.SRV;
 
 		return UserDataElement;
 	}
@@ -417,26 +354,16 @@ public:
 		return reinterpret_cast<size_t>(&UniquePointer);
 	}
 
-private:
-	virtual void CreateRenderThreadResources() override
-	{
-		InstanceBuffer.bOwnedByEditor = Component->IsOwnedByEditor();
-	}
-	virtual void OnTransformChanged() override
-	{
-		TraversalOctree.Reset(new FLidarPointCloudTraversalOctree(Octree, Component->GetComponentTransform()));
-	}
+	virtual void UpdateRenderData(FLidarPointCloudProxyUpdateData InRenderData) override { RenderData = InRenderData; }
+
+public:
+	TSharedPtr<FLidarPointCloudSceneProxyWrapper, ESPMode::ThreadSafe> ProxyWrapper;
 
 private:
-	FLidarPointCloudInstanceBuffer InstanceBuffer;
+	FLidarPointCloudProxyUpdateData RenderData;
 	mutable TArray<FLidarPointCloudBatchElementUserData> UserData;
-	TUniquePtr<FLidarPointCloudTraversalOctree> TraversalOctree;
-	TArray<FLidarPointCloudTraversalOctreeNode*> SelectedNodes;
-
-	bool bReadyToRender;
 
 	ULidarPointCloudComponent* Component;
-	FLidarPointCloudOctree* Octree;
 	FMaterialRelevance MaterialRelevance;
 	AActor* Owner;
 
@@ -445,10 +372,11 @@ private:
 
 FPrimitiveSceneProxy* ULidarPointCloudComponent::CreateSceneProxy()
 {
-	FPrimitiveSceneProxy* Proxy = nullptr;
+	FLidarPointCloudSceneProxy* Proxy = nullptr;
 	if (PointCloud)
 	{
 		Proxy = new FLidarPointCloudSceneProxy(this);
+		FLidarPointCloudLODManager::RegisterProxy(this, Proxy->ProxyWrapper);
 	}
 	return Proxy;
 }
@@ -506,4 +434,4 @@ void ULidarPointCloud::ReleaseCollisionRendering()
 				ReleaseCollisionRendering();
 			});
 	}
-}
\ No newline at end of file
+}
diff --git a/Engine/Plugins/Enterprise/LidarPointCloud/Source/LidarPointCloudRuntime/Public/LidarPointCloud.h b/Engine/Plugins/Enterprise/LidarPointCloud/Source/LidarPointCloudRuntime/Public/LidarPointCloud.h
--- a/Engine/Plugins/Enterprise/LidarPointCloud/Source/LidarPointCloudRuntime/Public/LidarPointCloud.h
+++ b/Engine/Plugins/Enterprise/LidarPointCloud/Source/LidarPointCloudRuntime/Public/LidarPointCloud.h
@@ -76,6 +76,9 @@ private:
 	void UpdateStatus();
 };
 
+/** Used to notify the component it should refresh its state. */
+DECLARE_EVENT(ULidarPointCloud, FOnPointCloudChanged);
+
 /**
  * Represents the Point Cloud asset
  */
@@ -108,6 +111,10 @@ public:
 	UPROPERTY()
 	FDoubleVector OriginalCoordinates;
 
+	/** Contains an offset to be added to all points when rendering */
+	UPROPERTY()
+	FDoubleVector LocationOffset;
+
 	/** Required for file versioning */
 	static const FGuid PointCloudFileGUID;
 	static const int32 PointCloudFileVersion;
@@ -138,19 +145,23 @@ private:
 	/** Used for collision building */
 	FThreadSafeBool bCollisionBuildInProgress;
 
+	FOnPointCloudChanged OnPointCloudRebuiltEvent;
+	FOnPointCloudChanged OnPointCloudUpdateCollisionEvent;
+	FOnPointCloudChanged OnPreSaveCleanupEvent;
+
 public:
 	ULidarPointCloud();
 
-	/** Used to notify the component it should refresh its state. */
-	DECLARE_EVENT(ULidarPointCloud, FOnPointCloudChanged);
 	FOnPointCloudChanged& OnPointCloudRebuilt() { return OnPointCloudRebuiltEvent; }
 	FOnPointCloudChanged& OnPointCloudCollisionUpdated() { return OnPointCloudUpdateCollisionEvent; }
+	FOnPointCloudChanged& OnPreSaveCleanup() { return OnPreSaveCleanupEvent; }
 
 	// Begin UObject Interface.
 	virtual void Serialize(FArchive& Ar) override;
 	virtual void PostLoad() override;
 	virtual void GetAssetRegistryTags(TArray<FAssetRegistryTag>& OutTags) const override;
 	virtual void BeginDestroy() override;
+	virtual void PreSave(const class ITargetPlatform* TargetPlatform) override;
 #if WITH_EDITOR
 	virtual void PostEditChangeProperty(FPropertyChangedEvent & PropertyChangedEvent) override;
 #endif
@@ -173,14 +184,14 @@ public:
 	FString GetSourcePath() const { return SourcePath.FilePath; }
 
 	UFUNCTION(BlueprintPure, Category = "Lidar Point Cloud")
-	FBox GetBounds() const { return Octree.GetBounds(); }
+	FBox GetBounds() const { return Octree.GetBounds().ShiftBy(LocationOffset.ToVector()); }
 
-	UFUNCTION(BlueprintPure, Category = "Lidar Point Cloud")
-	FBox GetPointsBounds() const { return Octree.GetPointsBounds(); }
+	/** Returns the Cloud's offset from the 0,0,0 coordinate */
+	FDoubleVector GetLocationOffset() const { return LocationOffset; }
 
 	/** Recalculates and updates points bounds. */
 	UFUNCTION(BlueprintCallable, Category = "Lidar Point Cloud")
-	void RefreshPointsBounds() { Octree.RefreshPointsBounds(); }
+	void RefreshBounds() { Octree.RefreshBounds(); }
 
 	/** Returns true, if the Octree has collision built */
 	UFUNCTION(BlueprintPure, Category = "Lidar Point Cloud")
@@ -196,84 +207,96 @@ public:
 	void GetPoints(TArray<FLidarPointCloudPoint*>& Points, int64 StartIndex = 0, int64 Count = -1) { Octree.GetPoints(Points, StartIndex, Count); }
 
 	/** Populates the array with the list of points within the given sphere. */
-	void GetPointsInSphere(TArray<FLidarPointCloudPoint*>& SelectedPoints, const FVector& Center, const float& Radius, const bool& bVisibleOnly) { GetPointsInSphere(SelectedPoints, FSphere(Center, Radius), bVisibleOnly); }
-	void GetPointsInSphere(TArray<FLidarPointCloudPoint*>& SelectedPoints, const FSphere& Sphere, const bool& bVisibleOnly) { Octree.GetPointsInSphere(SelectedPoints, Sphere, bVisibleOnly); }
+	void GetPointsInSphere(TArray<FLidarPointCloudPoint*>& SelectedPoints, FSphere Sphere, const bool& bVisibleOnly)
+	{
+		Sphere.Center -= LocationOffset.ToVector();
+		Octree.GetPointsInSphere(SelectedPoints, Sphere, bVisibleOnly);
+	}
 
 	/** Populates the array with the list of points within the given box. */
-	void GetPointsInBox(TArray<FLidarPointCloudPoint*>& SelectedPoints, const FVector& Center, const FVector& Extent, const bool& bVisibleOnly) { GetPointsInBox(SelectedPoints, FBox(Center - Extent, Center + Extent), bVisibleOnly); }
-	void GetPointsInBox(TArray<FLidarPointCloudPoint*>& SelectedPoints, const FBox& Box, const bool& bVisibleOnly) { Octree.GetPointsInBox(SelectedPoints, Box, bVisibleOnly); }
+	void GetPointsInBox(TArray<FLidarPointCloudPoint*>& SelectedPoints, const FBox& Box, const bool& bVisibleOnly) { Octree.GetPointsInBox(SelectedPoints, Box.ShiftBy(-LocationOffset.ToVector()), bVisibleOnly); }
 
-	/** Populates the array with the list of points within the given frustum. */
+	/**
+	 * Populates the array with the list of points within the given frustum.
+	 * Frustum is assumed to include the LocationOffset of the asset
+	 */
 	void GetPointsInFrustum(TArray<FLidarPointCloudPoint*>& SelectedPoints, const FConvexVolume& Frustum, const bool& bVisibleOnly) { Octree.GetPointsInFrustum(SelectedPoints, Frustum, bVisibleOnly); }
 
-	/** Returns an array with copies of points from the tree */
+	/**
+	 * Returns an array with copies of points from the tree
+	 * If ReturnWorldSpace is selected, the points' locations will be converted into absolute value, otherwise they will be relative to the center of the cloud.
+	 */
 	UFUNCTION(BlueprintPure, Category = "Lidar Point Cloud")
-	TArray<FLidarPointCloudPoint> GetPointsAsCopies(int32 StartIndex = 0, int32 Count = -1) const;
-	void GetPointsAsCopies(TArray<FLidarPointCloudPoint>& Points, int64 StartIndex = 0, int64 Count = -1) const { Octree.GetPointsAsCopies(Points, StartIndex, Count); }
+	TArray<FLidarPointCloudPoint> GetPointsAsCopies(bool bReturnWorldSpace, int32 StartIndex = 0, int32 Count = -1) const;
+	void GetPointsAsCopies(TArray<FLidarPointCloudPoint>& Points, bool bReturnWorldSpace, int64 StartIndex = 0, int64 Count = -1) const
+	{
+		FTransform LocalToWorld(LocationOffset.ToVector());
+		Octree.GetPointsAsCopies(Points, bReturnWorldSpace ? &LocalToWorld : nullptr, StartIndex, Count);
+	}
 	
-	/** Returns an array with copies of points within the given sphere */
+	/**
+	 * Returns an array with copies of points within the given sphere
+	 * If ReturnWorldSpace is selected, the points' locations will be converted into absolute value, otherwise they will be relative to the center of the cloud.
+	 */
 	UFUNCTION(BlueprintPure, Category = "Lidar Point Cloud")
-	TArray<FLidarPointCloudPoint> GetPointsInSphereAsCopies(FVector Center, float Radius, bool bVisibleOnly);
-	void GetPointsInSphereAsCopies(TArray<FLidarPointCloudPoint>& SelectedPoints, const FSphere& Sphere, const bool& bVisibleOnly) { Octree.GetPointsInSphereAsCopies(SelectedPoints, Sphere, bVisibleOnly); }
+	TArray<FLidarPointCloudPoint> GetPointsInSphereAsCopies(FVector Center, float Radius, bool bVisibleOnly, bool bReturnWorldSpace);
+	void GetPointsInSphereAsCopies(TArray<FLidarPointCloudPoint>& SelectedPoints, FSphere Sphere, const bool& bVisibleOnly, bool bReturnWorldSpace)
+	{
+		FTransform LocalToWorld(LocationOffset.ToVector());
+		Sphere.Center -= LocationOffset.ToVector();
+		Octree.GetPointsInSphereAsCopies(SelectedPoints, Sphere, bVisibleOnly, bReturnWorldSpace ? &LocalToWorld : nullptr);
+	}
 
-	/** Returns an array with copies of points within the given box */
+	/**
+	 * Returns an array with copies of points within the given box
+	 * If ReturnWorldSpace is selected, the points' locations will be converted into absolute value, otherwise they will be relative to the center of the cloud.
+	 */
 	UFUNCTION(BlueprintPure, Category = "Lidar Point Cloud")
-	TArray<FLidarPointCloudPoint> GetPointsInBoxAsCopies(FVector Center, FVector Extent, bool bVisibleOnly);
-	void GetPointsInBoxAsCopies(TArray<FLidarPointCloudPoint>& SelectedPoints, const FBox& Box, const bool& bVisibleOnly) { Octree.GetPointsInBoxAsCopies(SelectedPoints, Box, bVisibleOnly); }
+	TArray<FLidarPointCloudPoint> GetPointsInBoxAsCopies(FVector Center, FVector Extent, bool bVisibleOnly, bool bReturnWorldSpace);
+	void GetPointsInBoxAsCopies(TArray<FLidarPointCloudPoint>& SelectedPoints, const FBox& Box, const bool& bVisibleOnly, bool bReturnWorldSpace)
+	{
+		FTransform LocalToWorld(LocationOffset.ToVector());
+		Octree.GetPointsInBoxAsCopies(SelectedPoints, Box.ShiftBy(-LocationOffset.ToVector()), bVisibleOnly, bReturnWorldSpace ? &LocalToWorld : nullptr);
+	}
 
 	/** Performs a raycast test against the point cloud. Returns the pointer if hit or nullptr otherwise. */
 	UFUNCTION(BlueprintPure, Category = "Lidar Point Cloud", meta = (Keywords = "raycast"))
 	bool LineTraceSingle(FVector Origin, FVector Direction, float Radius, bool bVisibleOnly, FLidarPointCloudPoint& PointHit);
-	FLidarPointCloudPoint* LineTraceSingle(const FLidarPointCloudRay& Ray, const float& Radius, const bool& bVisibleOnly) { return Octree.RaycastSingle(Ray, Radius, bVisibleOnly); }
+	FLidarPointCloudPoint* LineTraceSingle(const FLidarPointCloudRay& Ray, const float& Radius, const bool& bVisibleOnly) { return Octree.RaycastSingle(Ray.ShiftBy(-LocationOffset.ToVector()), Radius, bVisibleOnly); }
 
 	/**
 	 * Performs a raycast test against the point cloud.
 	 * Populates OutHits array with the results.
+	 * If ReturnWorldSpace is selected, the points' locations will be converted into absolute value, otherwise they will be relative to the center of the cloud.
 	 * Returns true it anything has been hit.
 	 */
 	UFUNCTION(BlueprintPure, Category = "Lidar Point Cloud", meta = (Keywords = "raycast"))
-	bool LineTraceMulti(FVector Origin, FVector Direction, float Radius, bool bVisibleOnly, TArray<FLidarPointCloudPoint>& OutHits) { return Octree.RaycastMulti(FLidarPointCloudRay(Origin, Direction), Radius, bVisibleOnly, OutHits); }
-	bool LineTraceMulti(const FLidarPointCloudRay& Ray, const float& Radius, bool bVisibleOnly, TArray<FLidarPointCloudPoint>& OutHits) { return Octree.RaycastMulti(Ray, Radius, bVisibleOnly, OutHits); }
-	bool LineTraceMulti(const FLidarPointCloudRay& Ray, const float& Radius, bool bVisibleOnly, TArray<FLidarPointCloudPoint*>& OutHits) { return Octree.RaycastMulti(Ray, Radius, bVisibleOnly, OutHits); }
+	bool LineTraceMulti(FVector Origin, FVector Direction, float Radius, bool bVisibleOnly, bool bReturnWorldSpace, TArray<FLidarPointCloudPoint>& OutHits) { return LineTraceMulti(FLidarPointCloudRay(Origin, Direction), Radius, bVisibleOnly, bReturnWorldSpace, OutHits); }
+	bool LineTraceMulti(const FLidarPointCloudRay& Ray, const float& Radius, bool bVisibleOnly, bool bReturnWorldSpace, TArray<FLidarPointCloudPoint>& OutHits)
+	{
+		FTransform LocalToWorld(LocationOffset.ToVector());
+		return Octree.RaycastMulti(Ray.ShiftBy(-LocationOffset.ToVector()), Radius, bVisibleOnly, bReturnWorldSpace ? &LocalToWorld : nullptr, OutHits);
+	}
+	bool LineTraceMulti(const FLidarPointCloudRay& Ray, const float& Radius, bool bVisibleOnly, TArray<FLidarPointCloudPoint*>& OutHits) { return Octree.RaycastMulti(Ray.ShiftBy(-LocationOffset.ToVector()), Radius, bVisibleOnly, OutHits); }
 
 	/** Sets visibility of points within the given sphere. */
 	UFUNCTION(BlueprintCallable, Category = "Lidar Point Cloud")
 	void SetVisibilityOfPointsInSphere(bool bNewVisibility, FVector Center, float Radius) { SetVisibilityOfPointsInSphere(bNewVisibility, FSphere(Center, Radius)); }
-	void SetVisibilityOfPointsInSphere(const bool& bNewVisibility, const FSphere& Sphere) { Octree.SetVisibilityOfPointsInSphere(bNewVisibility, Sphere); }
-
-	/**
-	 * Sets visibility of points within the given sphere.
-	 * Async version - does not wait for completion before returning from the call.
-	 */
-	UFUNCTION(BlueprintCallable, Category = "Lidar Point Cloud")
-	void SetVisibilityOfPointsInSphereAsync(bool bNewVisibility, FVector Center, float Radius) { SetVisibilityOfPointsInSphereAsync(bNewVisibility, FSphere(Center, Radius)); }
-	void SetVisibilityOfPointsInSphereAsync(const bool& bNewVisibility, const FSphere& Sphere, TFunction<void(void)> CompletionCallback = nullptr) { Octree.SetVisibilityOfPointsInSphereAsync(bNewVisibility, Sphere, MoveTemp(CompletionCallback)); }
+	void SetVisibilityOfPointsInSphere(const bool& bNewVisibility, FSphere Sphere)
+	{
+		Sphere.Center -= LocationOffset.ToVector();
+		Octree.SetVisibilityOfPointsInSphere(bNewVisibility, Sphere);
+	}
 
 	/** Sets visibility of points within the given box. */
 	UFUNCTION(BlueprintCallable, Category = "Lidar Point Cloud")
 	void SetVisibilityOfPointsInBox(bool bNewVisibility, FVector Center, FVector Extent) { SetVisibilityOfPointsInBox(bNewVisibility, FBox(Center - Extent, Center + Extent)); }
-	void SetVisibilityOfPointsInBox(const bool& bNewVisibility, const FBox& Box) { Octree.SetVisibilityOfPointsInBox(bNewVisibility, Box); }
-
-	/**
-	 * Sets visibility of points within the given box.
-	 * Async version - does not wait for completion before returning from the call.
-	 */
-	UFUNCTION(BlueprintCallable, Category = "Lidar Point Cloud")
-	void SetVisibilityOfPointsInBoxAsync(bool bNewVisibility, FVector Center, FVector Extent) { SetVisibilityOfPointsInBoxAsync(bNewVisibility, FBox(Center - Extent, Center + Extent)); }
-	void SetVisibilityOfPointsInBoxAsync(const bool& bNewVisibility, const FBox& Box, TFunction<void(void)> CompletionCallback = nullptr) { Octree.SetVisibilityOfPointsInBoxAsync(bNewVisibility, Box, MoveTemp(CompletionCallback)); }
+	void SetVisibilityOfPointsInBox(const bool& bNewVisibility, const FBox& Box) { Octree.SetVisibilityOfPointsInBox(bNewVisibility, Box.ShiftBy(-LocationOffset.ToVector())); }
 
 	/** Sets visibility of points hit by the given ray. */
 	UFUNCTION(BlueprintCallable, Category = "Lidar Point Cloud")
 	void SetVisibilityOfPointsByRay(bool bNewVisibility, FVector Origin, FVector Direction, float Radius) { SetVisibilityOfPointsByRay(bNewVisibility, FLidarPointCloudRay(Origin, Direction), Radius); }
-	void SetVisibilityOfPointsByRay(const bool& bNewVisibility, const FLidarPointCloudRay& Ray, const float& Radius) { Octree.SetVisibilityOfPointsByRay(bNewVisibility, Ray, Radius); }
-
-	/**
-	 * Sets visibility of points hit by the given ray.
-	 * Async version - does not wait for completion before returning from the call.
-	 */
-	UFUNCTION(BlueprintCallable, Category = "Lidar Point Cloud")
-	void SetVisibilityOfPointsByRayAsync(bool bNewVisibility, FVector Origin, FVector Direction, float Radius) { SetVisibilityOfPointsByRayAsync(bNewVisibility, FLidarPointCloudRay(Origin, Direction), Radius); }
-	void SetVisibilityOfPointsByRayAsync(const bool& bNewVisibility, const FLidarPointCloudRay& Ray, const float& Radius, TFunction<void(void)> CompletionCallback = nullptr) { Octree.SetVisibilityOfPointsByRayAsync(bNewVisibility, Ray, Radius, MoveTemp(CompletionCallback)); }
+	void SetVisibilityOfPointsByRay(const bool& bNewVisibility, const FLidarPointCloudRay& Ray, const float& Radius) { Octree.SetVisibilityOfPointsByRay(bNewVisibility, Ray.ShiftBy(-LocationOffset.ToVector()), Radius); }
 
 	/** Marks all points hidden */
 	UFUNCTION(BlueprintCallable, Category = "Lidar Point Cloud")
@@ -283,53 +306,23 @@ public:
 	UFUNCTION(BlueprintCallable, Category = "Lidar Point Cloud")
 	void UnhideAll() { Octree.UnhideAll(); }
 
-	/**
-	 * Marks all points visible
-	 * Async version - does not wait for completion before returning from the call.
-	 */
-	UFUNCTION(BlueprintCallable, Category = "Lidar Point Cloud")
-	void ResetVisibilityAsync() { Octree.ResetVisibilityAsync(); }
-	void ResetVisibilityAsync(TFunction<void(void)> CompletionCallback) { Octree.ResetVisibilityAsync(MoveTemp(CompletionCallback)); }
-
 	/** Executes the provided action on each of the points. */
 	void ExecuteActionOnAllPoints(TFunction<void(FLidarPointCloudPoint*)> Action, const bool& bVisibleOnly) { Octree.ExecuteActionOnAllPoints(MoveTemp(Action), bVisibleOnly); }
 
-	/**
-	 * Executes the provided action on each of the points.
-	 * Async version - does not wait for completion before returning from the call.
-	 */
-	void ExecuteActionOnAllPointsAsync(TFunction<void(FLidarPointCloudPoint*)> Action, const bool& bVisibleOnly, TFunction<void(void)> CompletionCallback = nullptr) { Octree.ExecuteActionOnAllPointsAsync(MoveTemp(Action), bVisibleOnly, MoveTemp(CompletionCallback)); }
-
 	/** Executes the provided action on each of the points within the given sphere. */
 	void ExecuteActionOnPointsInSphere(TFunction<void(FLidarPointCloudPoint*)> Action, const FVector& Center, const float& Radius, const bool& bVisibleOnly) { ExecuteActionOnPointsInSphere(MoveTemp(Action), FSphere(Center, Radius), bVisibleOnly); }
-	void ExecuteActionOnPointsInSphere(TFunction<void(FLidarPointCloudPoint*)> Action, const FSphere& Sphere, const bool& bVisibleOnly) { Octree.ExecuteActionOnPointsInSphere(MoveTemp(Action), Sphere, bVisibleOnly); }
-
-	/**
-	 * Executes the provided action on each of the points within the given sphere.
-	 * Async version - does not wait for completion before returning from the call.
-	 */
-	void ExecuteActionOnPointsInSphereAsync(TFunction<void(FLidarPointCloudPoint*)> Action, const FVector& Center, const float& Radius, const bool& bVisibleOnly) { ExecuteActionOnPointsInSphereAsync(MoveTemp(Action), FSphere(Center, Radius), bVisibleOnly); }
-	void ExecuteActionOnPointsInSphereAsync(TFunction<void(FLidarPointCloudPoint*)> Action, const FSphere& Sphere, const bool& bVisibleOnly, TFunction<void(void)> CompletionCallback = nullptr) { Octree.ExecuteActionOnPointsInSphereAsync(MoveTemp(Action), Sphere, bVisibleOnly, MoveTemp(CompletionCallback)); }
-
+	void ExecuteActionOnPointsInSphere(TFunction<void(FLidarPointCloudPoint*)> Action, FSphere Sphere, const bool& bVisibleOnly)
+	{
+		Sphere.Center -= LocationOffset.ToVector();
+		Octree.ExecuteActionOnPointsInSphere(MoveTemp(Action), Sphere, bVisibleOnly);
+	}
+	
 	/** Executes the provided action on each of the points within the given box. */
 	void ExecuteActionOnPointsInBox(TFunction<void(FLidarPointCloudPoint*)> Action, const FVector& Center, const FVector& Extent, const bool& bVisibleOnly) { ExecuteActionOnPointsInBox(MoveTemp(Action), FBox(Center - Extent, Center + Extent), bVisibleOnly); }
-	void ExecuteActionOnPointsInBox(TFunction<void(FLidarPointCloudPoint*)> Action, const FBox& Box, const bool& bVisibleOnly) { Octree.ExecuteActionOnPointsInBox(MoveTemp(Action), Box, bVisibleOnly); }
-
-	/**
-	 * Executes the provided action on each of the points within the given box.
-	 * Async version - does not wait for completion before returning from the call.
-	 */
-	void ExecuteActionOnPointsInBoxAsync(TFunction<void(FLidarPointCloudPoint*)> Action, const FVector& Center, const FVector& Extent, const bool& bVisibleOnly) { ExecuteActionOnPointsInBoxAsync(MoveTemp(Action), FBox(Center - Extent, Center + Extent), bVisibleOnly); }
-	void ExecuteActionOnPointsInBoxAsync(TFunction<void(FLidarPointCloudPoint*)> Action, const FBox& Box, const bool& bVisibleOnly, TFunction<void(void)> CompletionCallback = nullptr) { Octree.ExecuteActionOnPointsInBoxAsync(MoveTemp(Action), Box, bVisibleOnly, MoveTemp(CompletionCallback)); }
+	void ExecuteActionOnPointsInBox(TFunction<void(FLidarPointCloudPoint*)> Action, const FBox& Box, const bool& bVisibleOnly) { Octree.ExecuteActionOnPointsInBox(MoveTemp(Action), Box.ShiftBy(-LocationOffset.ToVector()), bVisibleOnly); }
 
 	/** Executes the provided action on each of the points hit by the given ray. */
-	void ExecuteActionOnPointsByRay(TFunction<void(FLidarPointCloudPoint*)> Action, const FLidarPointCloudRay& Ray, const float& Radius, bool bVisibleOnly) { Octree.ExecuteActionOnPointsByRay(MoveTemp(Action), Ray, Radius, bVisibleOnly); }
-
-	/**
-	 * Executes the provided action on each of the points hit by the given ray.
-	 * Async version - does not wait for completion before returning from the call.
-	 */
-	void ExecuteActionOnPointsByRayAsync(TFunction<void(FLidarPointCloudPoint*)> Action, const FLidarPointCloudRay& Ray, const float& Radius, bool bVisibleOnly, TFunction<void(void)> CompletionCallback = nullptr) { Octree.ExecuteActionOnPointsByRayAsync(MoveTemp(Action), Ray, Radius, bVisibleOnly, MoveTemp(CompletionCallback)); }
+	void ExecuteActionOnPointsByRay(TFunction<void(FLidarPointCloudPoint*)> Action, const FLidarPointCloudRay& Ray, const float& Radius, bool bVisibleOnly) { Octree.ExecuteActionOnPointsByRay(MoveTemp(Action), Ray.ShiftBy(-LocationOffset.ToVector()), Radius, bVisibleOnly); }
 
 	/** Applies the given color to all points */
 	UFUNCTION(BlueprintCallable, Category = "Lidar Point Cloud")
@@ -338,17 +331,21 @@ public:
 	/** Applies the given color to all points within the sphere */
 	UFUNCTION(BlueprintCallable, Category = "Lidar Point Cloud")
 	void ApplyColorToPointsInSphere(FColor NewColor, FVector Center, float Radius, bool bVisibleOnly) { ApplyColorToPointsInSphere(NewColor, FSphere(Center, Radius), bVisibleOnly); }
-	void ApplyColorToPointsInSphere(const FColor& NewColor, const FSphere& Sphere, const bool& bVisibleOnly) { Octree.ApplyColorToPointsInSphere(NewColor, Sphere, bVisibleOnly); }
+	void ApplyColorToPointsInSphere(const FColor& NewColor, FSphere Sphere, const bool& bVisibleOnly)
+	{
+		Sphere.Center -= LocationOffset.ToVector();
+		Octree.ApplyColorToPointsInSphere(NewColor, Sphere, bVisibleOnly);
+	}
 
 	/** Applies the given color to all points within the box */
 	UFUNCTION(BlueprintCallable, Category = "Lidar Point Cloud")
 	void ApplyColorToPointsInBox(FColor NewColor, FVector Center, FVector Extent, bool bVisibleOnly) { ApplyColorToPointsInBox(NewColor, FBox(Center - Extent, Center + Extent), bVisibleOnly); }
-	void ApplyColorToPointsInBox(const FColor& NewColor, const FBox& Box, const bool& bVisibleOnly) { Octree.ApplyColorToPointsInBox(NewColor, Box, bVisibleOnly); }
+	void ApplyColorToPointsInBox(const FColor& NewColor, const FBox& Box, const bool& bVisibleOnly) { Octree.ApplyColorToPointsInBox(NewColor, Box.ShiftBy(-LocationOffset.ToVector()), bVisibleOnly); }
 
 	/** Applies the given color to all points hit by the given ray */
 	UFUNCTION(BlueprintCallable, Category = "Lidar Point Cloud")
 	void ApplyColorToPointsByRay(FColor NewColor, FVector Origin, FVector Direction, float Radius, bool bVisibleOnly) { ApplyColorToPointsByRay(NewColor, FLidarPointCloudRay(Origin, Direction), Radius, bVisibleOnly); }
-	void ApplyColorToPointsByRay(const FColor& NewColor, const FLidarPointCloudRay& Ray, const float& Radius, bool bVisibleOnly) { Octree.ApplyColorToPointsByRay(NewColor, Ray, Radius, bVisibleOnly); }
+	void ApplyColorToPointsByRay(const FColor& NewColor, const FLidarPointCloudRay& Ray, const float& Radius, bool bVisibleOnly) { Octree.ApplyColorToPointsByRay(NewColor, Ray.ShiftBy(-LocationOffset.ToVector()), Radius, bVisibleOnly); }
 
 	/**
 	 * This should to be called if any manual modification to individual points' visibility has been made.
@@ -365,7 +362,11 @@ public:
 	 * Warning: Will erase all currently held data!
 	 */
 	UFUNCTION(BlueprintCallable, Category = "Lidar Point Cloud")
-	void Initialize(const FBox& NewBounds) { Octree.Initialize(NewBounds); }
+	void Initialize(const FBox& NewBounds)
+	{
+		LocationOffset = OriginalCoordinates = NewBounds.GetCenter();
+		Octree.Initialize(NewBounds.GetExtent());
+	}
 
 	/** Builds collision mesh for the cloud, using current collision settings */
 	UFUNCTION(BlueprintCallable, Category = "Lidar Point Cloud")
@@ -375,29 +376,34 @@ public:
 	UFUNCTION(BlueprintCallable, Category = "Lidar Point Cloud")
 	void RemoveCollision();
 
-	/** Centers all contained points */
+	/** Returns true, if the cloud is fully and persistently loaded. */
+	UFUNCTION(BlueprintPure, Category = "Lidar Point Cloud")
+	bool IsFullyLoaded() const { return Octree.IsFullyLoaded(); }
+
+	/** Persistently loads all nodes. */
 	UFUNCTION(BlueprintCallable, Category = "Lidar Point Cloud")
-	void CenterPoints();
+	void LoadAllNodes() { Octree.LoadAllNodes(); }
 
-	/** Restores original coordinates to all contained points */
+	/** Applies given offset to this point cloud. */
 	UFUNCTION(BlueprintCallable, Category = "Lidar Point Cloud")
-	void RestoreOriginalCoordinates();
+	void SetLocationOffset(FVector Offset) { SetLocationOffset(FDoubleVector(Offset)); }
+	void SetLocationOffset(FDoubleVector Offset);
 
-	/**
-	 * Applies given offset to all contained points.
-	 * If bRefreshPointsBounds is set to false, make sure you call RefreshPointsBounds() manually or cloud centering may not work correctly.
-	 */
+	/** Centers this cloud */
+	UFUNCTION(BlueprintCallable, Category = "Lidar Point Cloud")
+	void CenterPoints() { SetLocationOffset(FDoubleVector::ZeroVector); }
+
+	/** Restores original coordinates */
 	UFUNCTION(BlueprintCallable, Category = "Lidar Point Cloud")
-	void ShiftPointsBy(FVector Offset, bool bRefreshPointsBounds) { ShiftPointsBy(FDoubleVector(Offset), bRefreshPointsBounds); }
-	void ShiftPointsBy(FDoubleVector Offset, bool bRefreshPointsBounds);
+	void RestoreOriginalCoordinates() { SetLocationOffset(OriginalCoordinates); }
 
 	/** Returns true, if the cloud has been centered. */
 	UFUNCTION(BlueprintCallable, Category = "Lidar Point Cloud")
-	bool IsCentered() const;
+	bool IsCentered() const { return LocationOffset.IsNearlyZero(0.1f); }
 
 	/** Re-imports the cloud from it's original source file, overwriting any current point information. */
-	UFUNCTION(BlueprintCallable, Category = "Lidar Point Cloud")
-	void Reimport(bool bUseAsync) { Reimport(FLidarPointCloudAsyncParameters(bUseAsync)); }
+	UFUNCTION(BlueprintCallable, Category = "Lidar Point Cloud", meta = (Latent, WorldContext = "WorldContextObject", LatentInfo = "LatentInfo", ExpandEnumAsExecs = "AsyncMode"))
+	void Reimport(UObject* WorldContextObject, bool bUseAsync, FLatentActionInfo LatentInfo, ELidarPointCloudAsyncMode& AsyncMode, float& Progress);
 	void Reimport(const FLidarPointCloudAsyncParameters& AsyncParameters);
 
 	/**
@@ -410,99 +416,88 @@ public:
 
 	/**
 	 * Inserts the given point into the Octree structure.
-	 * If bRefreshPointsBounds is set to false, make sure you call RefreshPointsBounds() manually or cloud centering may not work correctly.
+	 * If bRefreshPointsBounds is set to false, make sure you call RefreshBounds() manually or cloud centering may not work correctly.
 	 */
 	UFUNCTION(BlueprintCallable, Category = "Lidar Point Cloud")
-	void InsertPoint(const FLidarPointCloudPoint& Point, ELidarPointCloudDuplicateHandling DuplicateHandling, bool bRefreshPointsBounds);
+	void InsertPoint(const FLidarPointCloudPoint& Point, ELidarPointCloudDuplicateHandling DuplicateHandling, bool bRefreshPointsBounds, const FVector& Translation);
 
 	/**
 	 * Inserts group of points into the Octree structure, multi-threaded.
-	 * If bRefreshPointsBounds is set to false, make sure you call RefreshPointsBounds() manually or cloud centering may not work correctly.
+	 * If bRefreshPointsBounds is set to false, make sure you call RefreshBounds() manually or cloud centering may not work correctly.
 	 */
 	UFUNCTION(BlueprintCallable, Category = "Lidar Point Cloud")
-	void InsertPoints(const TArray<FLidarPointCloudPoint>& Points, ELidarPointCloudDuplicateHandling DuplicateHandling, bool bRefreshPointsBounds) { InsertPoints(Points.GetData(), Points.Num(), DuplicateHandling, bRefreshPointsBounds); }
+	void InsertPoints(const TArray<FLidarPointCloudPoint>& Points, ELidarPointCloudDuplicateHandling DuplicateHandling, bool bRefreshPointsBounds, const FVector& Translation) { InsertPoints(Points.GetData(), Points.Num(), DuplicateHandling, bRefreshPointsBounds, Translation); }
 
 	/**
 	 * Inserts group of points into the Octree structure, multi-threaded.
-	 * If bRefreshPointsBounds is set to false, make sure you call RefreshPointsBounds() manually or cloud centering may not work correctly.
+	 * If bRefreshPointsBounds is set to false, make sure you call RefreshBounds() manually or cloud centering may not work correctly.
 	 * Can be optionally passed a cancellation pointer - if it ever becomes non-null with value of true, process will be canceled.
 	 * May also provide progress callback, called approximately every 1% of progress.
 	 * Returns false if canceled.
 	 */
 	template<typename T>
-	bool InsertPoints(T InPoints, const int64& Count, ELidarPointCloudDuplicateHandling DuplicateHandling, bool bRefreshPointsBounds, FThreadSafeBool* bCanceled = nullptr, TFunction<void(float)> ProgressCallback = TFunction<void(float)>())
+	bool InsertPoints(T InPoints, const int64& Count, ELidarPointCloudDuplicateHandling DuplicateHandling, bool bRefreshPointsBounds, const FVector& Translation, FThreadSafeBool* bCanceled = nullptr, TFunction<void(float)> ProgressCallback = TFunction<void(float)>())
 	{
 		FScopeLock Lock(&Octree.DataLock);
-		return InsertPoints_NoLock(InPoints, Count, DuplicateHandling, bRefreshPointsBounds, bCanceled, MoveTemp(ProgressCallback));
+		return InsertPoints_NoLock(InPoints, Count, DuplicateHandling, bRefreshPointsBounds, Translation, bCanceled, MoveTemp(ProgressCallback));
 	}
 	template<typename T>
-	bool InsertPoints_NoLock(T InPoints, const int64& Count, ELidarPointCloudDuplicateHandling DuplicateHandling, bool bRefreshPointsBounds, FThreadSafeBool* bCanceled = nullptr, TFunction<void(float)> ProgressCallback = TFunction<void(float)>());
+	bool InsertPoints_NoLock(T InPoints, const int64& Count, ELidarPointCloudDuplicateHandling DuplicateHandling, bool bRefreshPointsBounds, const FVector& Translation, FThreadSafeBool* bCanceled = nullptr, TFunction<void(float)> ProgressCallback = TFunction<void(float)>());
 
-	/**
-	 * Attempts to remove the given point.
-	 * If bRefreshPointsBounds is set to false, make sure you call RefreshPointsBounds() manually or cloud centering may not work correctly.
-	 */
+	/** Attempts to remove the given point. */
 	UFUNCTION(BlueprintCallable, Category = "Lidar Point Cloud")
-	void RemovePoint(FLidarPointCloudPoint Point, bool bRefreshPointsBounds)
+	void RemovePoint(FLidarPointCloudPoint Point)
 	{
 		FScopeLock Lock(&Octree.DataLock);
-		Octree.RemovePoint(Point, bRefreshPointsBounds);
+		Octree.RemovePoint(Point);
 	}
-	void RemovePoint_NoLock(FLidarPointCloudPoint Point, bool bRefreshPointsBounds)
+	void RemovePoint_NoLock(FLidarPointCloudPoint Point)
 	{
-		Octree.RemovePoint(Point, bRefreshPointsBounds);
+		Octree.RemovePoint(Point);
 	}
-	void RemovePoint(const FLidarPointCloudPoint* Point, bool bRefreshPointsBounds)
+	void RemovePoint(const FLidarPointCloudPoint* Point)
 	{
 		FScopeLock Lock(&Octree.DataLock);
-		RemovePoint_NoLock(Point, bRefreshPointsBounds);
+		RemovePoint_NoLock(Point);
 	}
-	void RemovePoint_NoLock(const FLidarPointCloudPoint* Point, bool bRefreshPointsBounds)
+	void RemovePoint_NoLock(const FLidarPointCloudPoint* Point)
 	{
-		Octree.RemovePoint(Point, bRefreshPointsBounds);
+		Octree.RemovePoint(Point);
 	}
 
-	/**
-	 * Removes points in bulk
-	 * If bRefreshPointsBounds is set to false, make sure you call RefreshPointsBounds() manually or cloud centering may not work correctly.
-	 */
-	void RemovePoints(TArray<FLidarPointCloudPoint*>& Points, bool bRefreshPointsBounds)
+	/** Removes points in bulk */
+	void RemovePoints(TArray<FLidarPointCloudPoint*>& Points)
 	{
 		FScopeLock Lock(&Octree.DataLock);
-		RemovePoints_NoLock(Points, bRefreshPointsBounds);
+		RemovePoints_NoLock(Points);
 	}
-	void RemovePoints_NoLock(TArray<FLidarPointCloudPoint*>& Points, bool bRefreshPointsBounds)
+	void RemovePoints_NoLock(TArray<FLidarPointCloudPoint*>& Points)
 	{
-		Octree.RemovePoints(Points, bRefreshPointsBounds);
+		Octree.RemovePoints(Points);
 	}
 
-	/**
-	 * Removes all points within the given sphere
-	 * If bRefreshPointsBounds is set to false, make sure you call RefreshPointsBounds() manually or cloud centering may not work correctly.
-	 */
+	/** Removes all points within the given sphere  */
 	UFUNCTION(BlueprintCallable, Category = "Lidar Point Cloud")
-	void RemovePointsInSphere(FVector Center, float Radius, bool bVisibleOnly, bool bRefreshPointsBounds) { RemovePointsInSphere(FSphere(Center, Radius), bVisibleOnly, bRefreshPointsBounds); }
-	void RemovePointsInSphere(const FSphere& Sphere, const bool& bVisibleOnly, bool bRefreshPointsBounds) { Octree.RemovePointsInSphere(Sphere, bVisibleOnly, bRefreshPointsBounds); }
+	void RemovePointsInSphere(FVector Center, float Radius, bool bVisibleOnly) { RemovePointsInSphere(FSphere(Center, Radius), bVisibleOnly); }
+	void RemovePointsInSphere(FSphere Sphere, const bool& bVisibleOnly)
+	{
+		Sphere.Center -= LocationOffset.ToVector();
+		Octree.RemovePointsInSphere(Sphere, bVisibleOnly);
+	}
 
-	/**
-	 * Removes all points within the given box
-	 * If bRefreshPointsBounds is set to false, make sure you call RefreshPointsBounds() manually or cloud centering may not work correctly.
-	 */
+	/** Removes all points within the given box */
 	UFUNCTION(BlueprintCallable, Category = "Lidar Point Cloud")
-	void RemovePointsInBox(FVector Center, FVector Extent, bool bVisibleOnly, bool bRefreshPointsBounds) { RemovePointsInBox(FBox(Center - Extent, Center + Extent), bVisibleOnly, bRefreshPointsBounds); }
-	void RemovePointsInBox(const FBox& Box, const bool& bVisibleOnly, bool bRefreshPointsBounds) { Octree.RemovePointsInBox(Box, bVisibleOnly, bRefreshPointsBounds); }
+	void RemovePointsInBox(FVector Center, FVector Extent, bool bVisibleOnly) { RemovePointsInBox(FBox(Center - Extent, Center + Extent), bVisibleOnly); }
+	void RemovePointsInBox(const FBox& Box, const bool& bVisibleOnly) { Octree.RemovePointsInBox(Box.ShiftBy(-LocationOffset.ToVector()), bVisibleOnly); }
 
-	/**
-	 * Removes all points hit by the given ray
-	 * If bRefreshPointsBounds is set to false, make sure you call RefreshPointsBounds() manually or cloud centering may not work correctly.
-	 */
+	/** Removes all points hit by the given ray */
 	UFUNCTION(BlueprintCallable, Category = "Lidar Point Cloud")
-	void RemovePointsByRay(FVector Origin, FVector Direction, float Radius, bool bVisibleOnly, bool bRefreshPointsBounds) { RemovePointsByRay(FLidarPointCloudRay(Origin, Direction), Radius, bVisibleOnly, bRefreshPointsBounds); }
-	void RemovePointsByRay(const FLidarPointCloudRay& Ray, const float& Radius, const bool& bVisibleOnly, bool bRefreshPointsBounds) { Octree.RemovePointsByRay(Ray, Radius, bVisibleOnly, bRefreshPointsBounds); }
+	void RemovePointsByRay(FVector Origin, FVector Direction, float Radius, bool bVisibleOnly) { RemovePointsByRay(FLidarPointCloudRay(Origin, Direction), Radius, bVisibleOnly); }
+	void RemovePointsByRay(const FLidarPointCloudRay& Ray, const float& Radius, const bool& bVisibleOnly) { Octree.RemovePointsByRay(Ray.ShiftBy(-LocationOffset.ToVector()), Radius, bVisibleOnly); }
 
 	/** Removes all hidden points */
 	UFUNCTION(BlueprintCallable, Category = "Lidar Point Cloud")
-	void RemoveHiddenPoints(bool bRefreshPointsBounds) { Octree.RemoveHiddenPoints(bRefreshPointsBounds); }
+	void RemoveHiddenPoints() { Octree.RemoveHiddenPoints(); }
 
 	/** Reinitializes the cloud with the new set of data. */
 	UFUNCTION(BlueprintCallable, Category = "Lidar Point Cloud")
@@ -528,6 +523,26 @@ public:
 
 	UBodySetup* GetBodySetup();
 
+	//~ Begin Deprecated
+	UFUNCTION(BlueprintCallable, Category = "Lidar Point Cloud", meta = (DeprecatedFunction, DeprecationMessage = "Use SetLocationOffset() instead"))
+	void ShiftPointsBy(FVector Offset, bool bRefreshPointsBounds) { SetLocationOffset(LocationOffset + Offset); }
+
+	UFUNCTION(BlueprintPure, Category = "Lidar Point Cloud", meta = (DeprecatedFunction, DeprecationMessage = "Use GetBounds() instead"))
+	FBox GetPointsBounds() const { return GetBounds(); }
+
+	UFUNCTION(BlueprintCallable, Category = "Lidar Point Cloud", meta = (DeprecatedFunction, DeprecationMessage = "Async methods are no longer provided out-of-the box"))
+	void SetVisibilityOfPointsInSphereAsync(bool bNewVisibility, FVector Center, float Radius) { SetVisibilityOfPointsInSphere(bNewVisibility, FSphere(Center, Radius)); }
+
+	UFUNCTION(BlueprintCallable, Category = "Lidar Point Cloud", meta = (DeprecatedFunction, DeprecationMessage = "Async methods are no longer provided out-of-the box"))
+	void SetVisibilityOfPointsInBoxAsync(bool bNewVisibility, FVector Center, FVector Extent) { SetVisibilityOfPointsInBox(bNewVisibility, FBox(Center - Extent, Center + Extent)); }
+
+	UFUNCTION(BlueprintCallable, Category = "Lidar Point Cloud", meta = (DeprecatedFunction, DeprecationMessage = "Async methods are no longer provided out-of-the box"))
+	void SetVisibilityOfPointsByRayAsync(bool bNewVisibility, FVector Origin, FVector Direction, float Radius) { SetVisibilityOfPointsByRay(bNewVisibility, FLidarPointCloudRay(Origin, Direction), Radius); }
+
+	UFUNCTION(BlueprintCallable, Category = "Lidar Point Cloud", meta = (DeprecatedFunction, DeprecationMessage = "Async methods are no longer provided out-of-the box"))
+	void ResetVisibilityAsync() { UnhideAll(); }
+	//~ End Deprecated
+
 public:
 	/** Aligns provided clouds based on the relative offset between their Original Coordinates. Retains overall centering of the group. */
 	static void AlignClouds(TArray<ULidarPointCloud*> PointCloudsToAlign);
@@ -565,9 +580,6 @@ private:
 
 	void InitializeCollisionRendering();
 	void ReleaseCollisionRendering();
-
-	FOnPointCloudChanged OnPointCloudRebuiltEvent;
-	FOnPointCloudChanged OnPointCloudUpdateCollisionEvent;
 };
 
 USTRUCT(BlueprintType)
@@ -629,6 +641,10 @@ public:
 	UFUNCTION(BlueprintCallable, Category = "Lidar Point Cloud")
 	static bool ExportPointCloudToFile(ULidarPointCloud* PointCloud, const FString& Filename) { return PointCloud && PointCloud->Export(Filename); }
 
+	/** Aligns provided clouds based on the relative offset between their Original Coordinates. Retains overall centering of the group. */
+	UFUNCTION(BlueprintCallable, Category = "Lidar Point Cloud")
+	static void AlignClouds(TArray<ULidarPointCloud*> PointCloudsToAlign) { ULidarPointCloud::AlignClouds(PointCloudsToAlign); }
+
 	/** Does a collision trace along the given line and returns the first blocking hit encountered. */
 	UFUNCTION(BlueprintPure, Category = "Lidar Point Cloud", meta = (WorldContext = "WorldContextObject", DisplayName = "LineTraceForLidarPointCloud", Keywords = "raycast"))
 	static bool LineTraceSingle(UObject* WorldContextObject, FVector Origin, FVector Direction, float Radius, bool bVisibleOnly, FLidarPointCloudTraceHit& Hit);
diff --git a/Engine/Plugins/Enterprise/LidarPointCloud/Source/LidarPointCloudRuntime/Public/LidarPointCloudComponent.h b/Engine/Plugins/Enterprise/LidarPointCloud/Source/LidarPointCloudRuntime/Public/LidarPointCloudComponent.h
--- a/Engine/Plugins/Enterprise/LidarPointCloud/Source/LidarPointCloudRuntime/Public/LidarPointCloudComponent.h
+++ b/Engine/Plugins/Enterprise/LidarPointCloud/Source/LidarPointCloudRuntime/Public/LidarPointCloudComponent.h
@@ -49,26 +49,11 @@ private:
 	UMaterialInterface* CustomMaterial;
 
 public:
-	/**
-	 * Determines the maximum number of points to be visible on the screen.
-	 * Higher values will produce better image quality, but will require faster hardware.
-	 */
-	UPROPERTY(EditAnywhere, BlueprintReadWrite, Category = "Performance", meta = (ClampMin = "0"))
-	int32 PointBudget;
-
-	/**
-	 * Determines the preference towards selecting nodes closer to screen center.
-	 * Larger values assign more priority towards screen center.
-	 * 0 to disable.
-	 */
-	UPROPERTY(EditAnywhere, BlueprintReadWrite, AdvancedDisplay, Category = "Performance", meta = (ClampMin = "0.0"))
-	float ScreenCenterImportance;
-
 	/**
 	 * Determines the minimum screen size for the node to be rendered.
 	 * Lower values will produce farther view distance, at the cost of higher CPU usage.
 	 */
-	UPROPERTY(EditAnywhere, BlueprintReadWrite, AdvancedDisplay, Category = "Performance", meta = (ClampMin = "0"))
+	UPROPERTY(EditAnywhere, BlueprintReadWrite, Category = "Performance", meta = (ClampMin = "0"))
 	float MinScreenSize;
 
 	/**
@@ -159,8 +144,8 @@ private:
 	UPROPERTY(Transient)
 	UMaterialInterface* MasterMaterialMasked;
 
-	/** For asset editor use only */
-	bool bOwnedByEditor;
+	/** Pointer to the viewport client of the owning editor, or null, if this is a game object. */
+	TWeakPtr<FViewportClient> OwningViewportClient;
 
 public:
 	ULidarPointCloudComponent();
@@ -168,7 +153,8 @@ public:
 	UFUNCTION(BlueprintPure, Category = "Components|LidarPointCloud")
 	ULidarPointCloud* GetPointCloud() const { return PointCloud; }
 
-	FORCEINLINE bool IsOwnedByEditor() const { return bOwnedByEditor; }
+	FORCEINLINE TWeakPtr<FViewportClient> GetOwningViewportClient() const { return OwningViewportClient; }
+	FORCEINLINE bool IsOwnedByEditor() const { return OwningViewportClient.IsValid(); }
 
 	/** Populates the array with the list of points within the given sphere. */
 	void GetPointsInSphere(TArray<FLidarPointCloudPoint*>& SelectedPoints, const FVector& Center, const float& Radius, const bool& bVisibleOnly) { GetPointsInSphere(SelectedPoints, FSphere(Center, Radius), bVisibleOnly); }
@@ -190,35 +176,43 @@ public:
 		}
 	}
 
-	/** Populates the array with copies of points within the given sphere. */
+	/**
+	 * Populates the array with copies of points within the given sphere.
+	 * If ReturnWorldSpace is selected, the points' locations will be converted into absolute value, otherwise they will be relative to the center of the cloud.
+	 */
 	UFUNCTION(BlueprintCallable, Category = "Lidar Point Cloud")
-	TArray<FLidarPointCloudPoint> GetPointsInSphereAsCopies(const FVector& Center, const float& Radius, const bool& bVisibleOnly)
+	TArray<FLidarPointCloudPoint> GetPointsInSphereAsCopies(FVector Center, float Radius, bool bVisibleOnly, bool bReturnWorldSpace)
 	{
 		TArray<FLidarPointCloudPoint> Points;
-		GetPointsInSphereAsCopies(Points, FSphere(Center, Radius), bVisibleOnly);
+		GetPointsInSphereAsCopies(Points, FSphere(Center, Radius), bVisibleOnly, bReturnWorldSpace);
 		return Points;
 	}
-	void GetPointsInSphereAsCopies(TArray<FLidarPointCloudPoint>& SelectedPoints, const FSphere& Sphere, const bool& bVisibleOnly)
+	void GetPointsInSphereAsCopies(TArray<FLidarPointCloudPoint>& SelectedPoints, const FSphere& Sphere, const bool& bVisibleOnly, const bool& bReturnWorldSpace)
 	{
 		if (PointCloud)
 		{
-			PointCloud->GetPointsInSphereAsCopies(SelectedPoints, Sphere.TransformBy(GetComponentTransform().Inverse()), bVisibleOnly);
+			FTransform LocalToWorld = GetLocalToWorld();
+			PointCloud->Octree.GetPointsInSphereAsCopies(SelectedPoints, Sphere.TransformBy(LocalToWorld.Inverse()), bVisibleOnly, bReturnWorldSpace ? &LocalToWorld : nullptr);
 		}
 	}
 
-	/** Populates the array with copies of points within the given box. */
+	/**
+	 * Populates the array with copies of points within the given box.
+	 * If ReturnWorldSpace is selected, the points' locations will be converted into absolute value, otherwise they will be relative to the center of the cloud.
+	 */
 	UFUNCTION(BlueprintCallable, Category = "Lidar Point Cloud")
-	TArray<FLidarPointCloudPoint> GetPointsInBoxAsCopies(FVector Center, FVector Extent, bool bVisibleOnly)
+	TArray<FLidarPointCloudPoint> GetPointsInBoxAsCopies(FVector Center, FVector Extent, bool bVisibleOnly, bool bReturnWorldSpace)
 	{
 		TArray<FLidarPointCloudPoint> Points;
-		GetPointsInBoxAsCopies(Points, FBox(Center - Extent, Center + Extent), bVisibleOnly);
+		GetPointsInBoxAsCopies(Points, FBox(Center - Extent, Center + Extent), bVisibleOnly, bReturnWorldSpace);
 		return Points;
 	}
-	void GetPointsInBoxAsCopies(TArray<FLidarPointCloudPoint>& SelectedPoints, const FBox& Box, const bool& bVisibleOnly)
+	void GetPointsInBoxAsCopies(TArray<FLidarPointCloudPoint>& SelectedPoints, const FBox& Box, const bool& bVisibleOnly, const bool& bReturnWorldSpace)
 	{
 		if (PointCloud)
 		{
-			PointCloud->GetPointsInBoxAsCopies(SelectedPoints, Box.TransformBy(GetComponentTransform().Inverse()), bVisibleOnly);
+			FTransform LocalToWorld = GetLocalToWorld();
+			PointCloud->Octree.GetPointsInBoxAsCopies(SelectedPoints, Box.TransformBy(LocalToWorld.Inverse()), bVisibleOnly, bReturnWorldSpace ? &LocalToWorld : nullptr);
 		}
 	}
 
@@ -243,16 +237,23 @@ public:
 	/**
 	 * Performs a raycast test against the point cloud.
 	 * Populates OutHits array with the results.
+	 * If ReturnWorldSpace is selected, the points' locations will be converted into absolute value, otherwise they will be relative to the center of the cloud.
 	 * Returns true it anything has been hit.
 	 */
 	UFUNCTION(BlueprintPure, Category = "Lidar Point Cloud", meta = (Keywords = "raycast"))
-	bool LineTraceMulti(FVector Origin, FVector Direction, float Radius, bool bVisibleOnly, TArray<FLidarPointCloudPoint>& OutHits)
+	bool LineTraceMulti(FVector Origin, FVector Direction, float Radius, bool bVisibleOnly, bool bReturnWorldSpace, TArray<FLidarPointCloudPoint>& OutHits)
 	{
-		return PointCloud ? PointCloud->LineTraceMulti(FLidarPointCloudRay(Origin, Direction).TransformBy(GetComponentTransform().Inverse()), Radius, bVisibleOnly, OutHits) : false;
+		return LineTraceMulti(FLidarPointCloudRay(Origin, Direction).TransformBy(GetComponentTransform().Inverse()), Radius, bVisibleOnly, bReturnWorldSpace, OutHits);
 	}
-	bool LineTraceMulti(FLidarPointCloudRay Ray, float Radius, bool bVisibleOnly, TArray<FLidarPointCloudPoint>& OutHits)
+	bool LineTraceMulti(FLidarPointCloudRay Ray, float Radius, bool bVisibleOnly, bool bReturnWorldSpace, TArray<FLidarPointCloudPoint>& OutHits)
 	{
-		return PointCloud ? PointCloud->LineTraceMulti(Ray.TransformBy(GetComponentTransform().Inverse()), Radius, bVisibleOnly, OutHits) : false;
+		if (PointCloud)
+		{
+			FTransform LocalToWorld = GetLocalToWorld();
+			return PointCloud->Octree.RaycastMulti(Ray.TransformBy(LocalToWorld.Inverse()), Radius, bVisibleOnly, bReturnWorldSpace ? &LocalToWorld : nullptr, OutHits);
+		}
+
+		return false;
 	}
 	bool LineTraceMulti(FLidarPointCloudRay Ray, float Radius, bool bVisibleOnly, TArray<FLidarPointCloudPoint*>& OutHits)
 	{
@@ -270,20 +271,6 @@ public:
 		}
 	}
 
-	/**
-	 * Sets visibility of points within the given sphere.
-	 * Async version - does not wait for completion before returning from the call.
-	 */
-	UFUNCTION(BlueprintCallable, Category = "Lidar Point Cloud")
-	void SetVisibilityOfPointsInSphereAsync(bool bNewVisibility, FVector Center, float Radius) { SetVisibilityOfPointsInSphereAsync(bNewVisibility, FSphere(Center, Radius)); }
-	void SetVisibilityOfPointsInSphereAsync(const bool& bNewVisibility, const FSphere& Sphere)
-	{
-		if (PointCloud)
-		{
-			PointCloud->SetVisibilityOfPointsInSphereAsync(bNewVisibility, Sphere.TransformBy(GetComponentTransform().Inverse()));
-		}
-	}
-
 	/** Sets visibility of points within the given box. */
 	UFUNCTION(BlueprintCallable, Category = "Lidar Point Cloud")
 	void SetVisibilityOfPointsInBox(bool bNewVisibility, FVector Center, FVector Extent) { SetVisibilityOfPointsInBox(bNewVisibility, FBox(Center - Extent, Center + Extent)); }
@@ -295,20 +282,6 @@ public:
 		}
 	}
 
-	/**
-	 * Sets visibility of points within the given box.
-	 * Async version - does not wait for completion before returning from the call.
-	 */
-	UFUNCTION(BlueprintCallable, Category = "Lidar Point Cloud")
-	void SetVisibilityOfPointsInBoxAsync(bool bNewVisibility, FVector Center, FVector Extent) { SetVisibilityOfPointsInBoxAsync(bNewVisibility, FBox(Center - Extent, Center + Extent)); }
-	void SetVisibilityOfPointsInBoxAsync(const bool& bNewVisibility, const FBox& Box)
-	{
-		if (PointCloud)
-		{
-			PointCloud->SetVisibilityOfPointsInBoxAsync(bNewVisibility, Box.TransformBy(GetComponentTransform().Inverse()));
-		}
-	}
-
 	/** Sets visibility of points hit by the given ray. */
 	UFUNCTION(BlueprintCallable, Category = "Lidar Point Cloud")
 	void SetVisibilityOfPointsByRay(bool bNewVisibility, FVector Origin, FVector Direction, float Radius) { SetVisibilityOfPointsByRay(bNewVisibility, FLidarPointCloudRay(Origin, Direction), Radius); }
@@ -320,20 +293,6 @@ public:
 		}
 	}
 
-	/**
-	 * Sets visibility of points hit by the given ray.
-	 * Async version - does not wait for completion before returning from the call.
-	 */
-	UFUNCTION(BlueprintCallable, Category = "Lidar Point Cloud")
-	void SetVisibilityOfPointsByRayAsync(bool bNewVisibility, FVector Origin, FVector Direction, float Radius) { SetVisibilityOfPointsByRayAsync(bNewVisibility, FLidarPointCloudRay(Origin, Direction), Radius); }
-	void SetVisibilityOfPointsByRayAsync(bool bNewVisibility, FLidarPointCloudRay Ray, float Radius, TFunction<void(void)> CompletionCallback = nullptr)
-	{
-		if (PointCloud)
-		{
-			PointCloud->SetVisibilityOfPointsByRayAsync(bNewVisibility, Ray.TransformBy(GetComponentTransform().Inverse()), Radius, MoveTemp(CompletionCallback));
-		}
-	}
-
 	/** Executes the provided action on each of the points within the given sphere. */
 	void ExecuteActionOnPointsInSphere(TFunction<void(FLidarPointCloudPoint*)> Action, const FVector& Center, const float& Radius, const bool& bVisibleOnly) { ExecuteActionOnPointsInSphere(MoveTemp(Action), FSphere(Center, Radius), bVisibleOnly); }
 	void ExecuteActionOnPointsInSphere(TFunction<void(FLidarPointCloudPoint*)> Action, const FSphere& Sphere, const bool& bVisibleOnly)
@@ -344,19 +303,6 @@ public:
 		}
 	}
 
-	/**
-	 * Executes the provided action on each of the points within the given sphere.
-	 * Async version - does not wait for completion before returning from the call.
-	 */
-	void ExecuteActionOnPointsInSphereAsync(TFunction<void(FLidarPointCloudPoint*)> Action, const FVector& Center, const float& Radius, const bool& bVisibleOnly) { ExecuteActionOnPointsInSphereAsync(MoveTemp(Action), FSphere(Center, Radius), bVisibleOnly); }
-	void ExecuteActionOnPointsInSphereAsync(TFunction<void(FLidarPointCloudPoint*)> Action, const FSphere& Sphere, const bool& bVisibleOnly)
-	{
-		if (PointCloud)
-		{
-			PointCloud->ExecuteActionOnPointsInSphereAsync(MoveTemp(Action), Sphere.TransformBy(GetComponentTransform().Inverse()), bVisibleOnly);
-		}
-	}
-
 	/** Executes the provided action on each of the points within the given box. */
 	void ExecuteActionOnPointsInBox(TFunction<void(FLidarPointCloudPoint*)> Action, const FVector& Center, const FVector& Extent, const bool& bVisibleOnly) { ExecuteActionOnPointsInBox(MoveTemp(Action), FBox(Center - Extent, Center + Extent), bVisibleOnly); }
 	void ExecuteActionOnPointsInBox(TFunction<void(FLidarPointCloudPoint*)> Action, const FBox& Box, const bool& bVisibleOnly)
@@ -367,19 +313,6 @@ public:
 		}
 	}
 
-	/**
-	 * Executes the provided action on each of the points within the given box.
-	 * Async version - does not wait for completion before returning from the call.
-	 */
-	void ExecuteActionOnPointsInBoxAsync(TFunction<void(FLidarPointCloudPoint*)> Action, const FVector& Center, const FVector& Extent, const bool& bVisibleOnly) { ExecuteActionOnPointsInBoxAsync(MoveTemp(Action), FBox(Center - Extent, Center + Extent), bVisibleOnly); }
-	void ExecuteActionOnPointsInBoxAsync(TFunction<void(FLidarPointCloudPoint*)> Action, const FBox& Box, const bool& bVisibleOnly)
-	{
-		if (PointCloud)
-		{
-			PointCloud->ExecuteActionOnPointsInBoxAsync(MoveTemp(Action), Box.TransformBy(GetComponentTransform().Inverse()), bVisibleOnly);
-		}
-	}
-
 	/** Executes the provided action on each of the points hit by the given ray. */
 	void ExecuteActionOnPointsByRay(TFunction<void(FLidarPointCloudPoint*)> Action, FLidarPointCloudRay Ray, float Radius, bool bVisibleOnly)
 	{
@@ -389,18 +322,6 @@ public:
 		}
 	}
 
-	/**
-	 * Executes the provided action on each of the points hit by the given ray.
-	 * Async version - does not wait for completion before returning from the call.
-	 */
-	void ExecuteActionOnPointsByRayAsync(TFunction<void(FLidarPointCloudPoint*)> Action, FLidarPointCloudRay Ray, float Radius, bool bVisibleOnly, TFunction<void(void)> CompletionCallback = nullptr)
-	{
-		if (PointCloud)
-		{
-			PointCloud->ExecuteActionOnPointsByRayAsync(MoveTemp(Action), Ray.TransformBy(GetComponentTransform().Inverse()), Radius, bVisibleOnly, MoveTemp(CompletionCallback));
-		}
-	}
-
 	/** Applies the given color to all points within the sphere */
 	UFUNCTION(BlueprintCallable, Category = "Lidar Point Cloud")
 	void ApplyColorToPointsInSphere(FColor NewColor, FVector Center, float Radius, bool bVisibleOnly) { ApplyColorToPointsInSphere(NewColor, FSphere(Center, Radius), bVisibleOnly); }
@@ -434,48 +355,50 @@ public:
 		}
 	}
 	
-	/**
-	 * Removes all points within the given sphere
-	 * If bRefreshPointsBounds is set to false, make sure you call RefreshPointsBounds() manually or cloud centering may not work correctly.
-	 */
+	/** Removes all points within the given sphere  */
 	UFUNCTION(BlueprintCallable, Category = "Lidar Point Cloud")
-	void RemovePointsInSphere(FVector Center, float Radius, bool bVisibleOnly, bool bRefreshPointsBounds) { RemovePointsInSphere(FSphere(Center, Radius), bVisibleOnly, bRefreshPointsBounds); }
-	void RemovePointsInSphere(const FSphere& Sphere, const bool& bVisibleOnly, bool bRefreshPointsBounds)
+	void RemovePointsInSphere(FVector Center, float Radius, bool bVisibleOnly) { RemovePointsInSphere(FSphere(Center, Radius), bVisibleOnly); }
+	void RemovePointsInSphere(const FSphere& Sphere, const bool& bVisibleOnly)
 	{
 		if (PointCloud)
 		{
-			PointCloud->RemovePointsInSphere(Sphere.TransformBy(GetComponentTransform().Inverse()), bVisibleOnly, bRefreshPointsBounds);
+			PointCloud->RemovePointsInSphere(Sphere.TransformBy(GetComponentTransform().Inverse()), bVisibleOnly);
 		}
 	}
 
-	/**
-	 * Removes all points within the given box
-	 * If bRefreshPointsBounds is set to false, make sure you call RefreshPointsBounds() manually or cloud centering may not work correctly.
-	 */
+	/** Removes all points within the given box  */
 	UFUNCTION(BlueprintCallable, Category = "Lidar Point Cloud")
-	void RemovePointsInBox(FVector Center, FVector Extent, bool bVisibleOnly, bool bRefreshPointsBounds) { RemovePointsInBox(FBox(Center - Extent, Center + Extent), bVisibleOnly, bRefreshPointsBounds); }
-	void RemovePointsInBox(const FBox& Box, const bool& bVisibleOnly, bool bRefreshPointsBounds)
+	void RemovePointsInBox(FVector Center, FVector Extent, bool bVisibleOnly) { RemovePointsInBox(FBox(Center - Extent, Center + Extent), bVisibleOnly); }
+	void RemovePointsInBox(const FBox& Box, const bool& bVisibleOnly)
 	{
 		if (PointCloud)
 		{
-			PointCloud->RemovePointsInBox(Box.TransformBy(GetComponentTransform().Inverse()), bVisibleOnly, bRefreshPointsBounds);
+			PointCloud->RemovePointsInBox(Box.TransformBy(GetComponentTransform().Inverse()), bVisibleOnly);
 		}
 	}
 
-	/**
-	 * Removes all points hit by the given ray
-	 * If bRefreshPointsBounds is set to false, make sure you call RefreshPointsBounds() manually or cloud centering may not work correctly.
-	 */
+	/** Removes all points hit by the given ray  */
 	UFUNCTION(BlueprintCallable, Category = "Lidar Point Cloud")
-	void RemovePointsByRay(FVector Origin, FVector Direction, float Radius, bool bVisibleOnly, bool bRefreshPointsBounds) { RemovePointsByRay(FLidarPointCloudRay(Origin, Direction), Radius, bVisibleOnly, bRefreshPointsBounds); }
-	void RemovePointsByRay(FLidarPointCloudRay Ray, float Radius, const bool& bVisibleOnly, bool bRefreshPointsBounds)
+	void RemovePointsByRay(FVector Origin, FVector Direction, float Radius, bool bVisibleOnly) { RemovePointsByRay(FLidarPointCloudRay(Origin, Direction), Radius, bVisibleOnly); }
+	void RemovePointsByRay(FLidarPointCloudRay Ray, float Radius, const bool& bVisibleOnly)
 	{
 		if (PointCloud)
 		{
-			PointCloud->RemovePointsByRay(Ray.TransformBy(GetComponentTransform().Inverse()), Radius, bVisibleOnly, bRefreshPointsBounds);
+			PointCloud->RemovePointsByRay(Ray.TransformBy(GetComponentTransform().Inverse()), Radius, bVisibleOnly);
 		}
 	}
 
+	//~ Begin Deprecated
+	UFUNCTION(BlueprintCallable, Category = "Lidar Point Cloud", meta = (DeprecatedFunction, DeprecationMessage = "Async methods are no longer provided out-of-the box"))
+	void SetVisibilityOfPointsInSphereAsync(bool bNewVisibility, FVector Center, float Radius) { SetVisibilityOfPointsInSphere(bNewVisibility, FSphere(Center, Radius)); }
+
+	UFUNCTION(BlueprintCallable, Category = "Lidar Point Cloud", meta = (DeprecatedFunction, DeprecationMessage = "Async methods are no longer provided out-of-the box"))
+	void SetVisibilityOfPointsInBoxAsync(bool bNewVisibility, FVector Center, FVector Extent) { SetVisibilityOfPointsInBox(bNewVisibility, FBox(Center - Extent, Center + Extent)); }
+
+	UFUNCTION(BlueprintCallable, Category = "Lidar Point Cloud", meta = (DeprecatedFunction, DeprecationMessage = "Async methods are no longer provided out-of-the box"))
+	void SetVisibilityOfPointsByRayAsync(bool bNewVisibility, FVector Origin, FVector Direction, float Radius) { SetVisibilityOfPointsByRay(bNewVisibility, FLidarPointCloudRay(Origin, Direction), Radius); }
+	//~ End Deprecated
+
 public:
 	UFUNCTION(BlueprintCallable, Category = "Components|LidarPointCloud")
 	void SetPointCloud(ULidarPointCloud *InPointCloud);
@@ -519,6 +442,16 @@ private:
 	virtual FBoxSphereBounds CalcBounds(const FTransform& LocalToWorld) const override;
 	// End USceneComponent Interface
 
+	FTransform GetLocalToWorld()
+	{
+		FTransform LocalToWorld = GetComponentTransform();
+		if (PointCloud)
+		{
+			LocalToWorld.AddToTranslation(PointCloud->LocationOffset.ToVector());
+		}
+		return LocalToWorld;
+	}
+
 	void UpdateMaterial();
 
 	void AttachPointCloudListener();
@@ -530,6 +463,6 @@ private:
 
 	friend class FPointCloudSceneProxy;
 #if WITH_EDITOR
-	friend class FLidarPointCloudEditorViewportClient;
+	friend class SLidarPointCloudEditorViewport;
 #endif
 };
