Merging

diff --git a/Engine/Plugins/Enterprise/LidarPointCloud/Source/LidarPointCloudRuntime/Public/LidarPointCloudOctree.h b/Engine/Plugins/Enterprise/LidarPointCloud/Source/LidarPointCloudRuntime/Public/LidarPointCloudOctree.h
--- a/Engine/Plugins/Enterprise/LidarPointCloud/Source/LidarPointCloudRuntime/Public/LidarPointCloudOctree.h
+++ b/Engine/Plugins/Enterprise/LidarPointCloud/Source/LidarPointCloudRuntime/Public/LidarPointCloudOctree.h
@@ -4,14 +4,17 @@
 
 #include "CoreMinimal.h"
 #include "LidarPointCloudShared.h"
-#include "Misc/ScopeLock.h"
 #include "LidarPointCloudSettings.h"
+#include "LidarPointCloudLODManager.h"
 #include "HAL/ThreadSafeCounter64.h"
+#include "Misc/ScopeLock.h"
 #include "ConvexVolume.h"
 #include "Interfaces/Interface_CollisionDataProvider.h"
 
+class ULidarPointCloud;
 class FLidarPointCloudOctree;
 struct FLidarPointCloudTraversalOctree;
+struct FLidarPointCloudTraversalOctreeNode;
 
 /**
  * WARNING: Exercise caution when modifying the contents of the Octree, as it may be in use by the Rendering Thread via FPointCloudSceneProxy
@@ -35,11 +38,10 @@ struct FLidarPointCloudTraversalOctree;
  */
 struct FLidarPointCloudOctreeNode
 {
-public:
-	mutable TArray<FLidarPointCloudPoint> AllocatedPoints;
-	mutable TArray<FLidarPointCloudPoint> PaddingPoints;
-
 private:
+	/** Stores the time, at which the BulkData needs to be released */
+	float BulkDataLifetime;
+
 	/** Depth of this node */
 	uint8 Depth;
 
@@ -61,18 +63,33 @@ private:
 
 	FCriticalSection MapLock;
 
+	/** Used for streaming the data from disk */
+	FLidarPointCloudBulkData BulkData;
+
+	/** Used to keep track, which data is available for rendering */
+	TAtomic<bool> bHasDataPending;
+
+	/** This is used to prevent nodes with changed content from being overwritten by consecutive streaming. */
+	TAtomic<bool> bCanReleaseData;
+
 public:
-	FORCEINLINE FLidarPointCloudOctreeNode() : FLidarPointCloudOctreeNode(nullptr, nullptr, 0) {}
-	FORCEINLINE FLidarPointCloudOctreeNode(FLidarPointCloudOctree* Tree, FLidarPointCloudOctreeNode* Parent, const uint8& Depth) : FLidarPointCloudOctreeNode(Tree, Parent, Depth, 0, FVector::ZeroVector) {}
-	FLidarPointCloudOctreeNode(FLidarPointCloudOctree* Tree, FLidarPointCloudOctreeNode* Parent, const uint8& Depth, const uint8& LocationInParent, const FVector& Center);
+	FORCEINLINE FLidarPointCloudOctreeNode() : FLidarPointCloudOctreeNode(nullptr, 0) {}
+	FORCEINLINE FLidarPointCloudOctreeNode(FLidarPointCloudOctree* Tree, const uint8& Depth) : FLidarPointCloudOctreeNode(Tree, Depth, 0, FVector::ZeroVector) {}
+	FLidarPointCloudOctreeNode(FLidarPointCloudOctree* Tree, const uint8& Depth, const uint8& LocationInParent, const FVector& Center);
 	~FLidarPointCloudOctreeNode();
 	FLidarPointCloudOctreeNode(const FLidarPointCloudOctreeNode&) = delete;
 	FLidarPointCloudOctreeNode(FLidarPointCloudOctreeNode&&) = delete;
 	FLidarPointCloudOctreeNode& operator=(const FLidarPointCloudOctreeNode&) = delete;
 	FLidarPointCloudOctreeNode& operator=(FLidarPointCloudOctreeNode&&) = delete;
 
+	/** Returns a pointer to the point data */
+	FORCEINLINE FLidarPointCloudPoint* GetData() const { return BulkData.GetData(); }
+
+	/** Returns a pointer to the point data and prevents it from being released */
+	FLidarPointCloudPoint* GetPersistentData() const;
+
 	/** Returns the sum of grid and padding points allocated to this node. */
-	int32 GetNumPoints() const { return AllocatedPoints.Num() + PaddingPoints.Num(); }
+	FORCEINLINE int32 GetNumPoints() const { return (int32)BulkData.GetElementCount(); }
 
 	/** Returns the sum of visible grid and padding points allocated to this node. */
 	int32 GetNumVisiblePoints() const { return NumVisiblePoints; }
@@ -89,8 +106,8 @@ public:
 	void UpdateNumVisiblePoints();
 
 	/** Attempts to insert given points to this node or passes it to the children, otherwise. */
-	void InsertPoints(FLidarPointCloudOctree* Tree, const FLidarPointCloudPoint* Points, const int32& Count, ELidarPointCloudDuplicateHandling DuplicateHandling);
-	void InsertPoints(FLidarPointCloudOctree* Tree, FLidarPointCloudPoint** Points, const int32& Count, ELidarPointCloudDuplicateHandling DuplicateHandling);
+	void InsertPoints(FLidarPointCloudOctree* Tree, const FLidarPointCloudPoint* Points, const int32& Count, ELidarPointCloudDuplicateHandling DuplicateHandling, const FVector& Translation);
+	void InsertPoints(FLidarPointCloudOctree* Tree, FLidarPointCloudPoint** Points, const int32& Count, ELidarPointCloudDuplicateHandling DuplicateHandling, const FVector& Translation);
 
 	/** Removes all points. */
 	void Empty(bool bRecursive = true);
@@ -99,9 +116,16 @@ public:
 	uint32 GetMaxDepth() const;
 
 	/** Returns the amount of memory used by this node */
-	int64 GetAllocatedSize(bool bRecursive) const;
+	int64 GetAllocatedSize(bool bRecursive, bool bIncludeBulkData) const;
 
-	void Serialize(FArchive& Ar, FLidarPointCloudOctree* Tree);
+	/** Returns true, if the node has its data loaded */
+	bool HasData() const { return BulkData.HasData(); }
+
+	/**
+	 * Releases the BulkData
+	 * If forced, the node will be released even if persistent
+	 */
+	void ReleaseData(bool bForce = false);
 
 	/**
 	 * Convenience function, to add point statistics to the Tree table.
@@ -109,8 +133,12 @@ public:
 	 */
 	void AddPointCount(FLidarPointCloudOctree* Tree, int32 PointCount = INT32_MIN);
 
+	/** Sorts the points by visibility (visible first) to optimize data processing and rendering */
+	void SortVisiblePoints();
+
 	friend FLidarPointCloudOctree;
-	friend struct FLidarPointCloudTraversalOctreeNode;
+	friend FLidarPointCloudTraversalOctree;
+	friend FLidarPointCloudTraversalOctreeNode;
 };
 
 /**
@@ -157,28 +185,35 @@ private:
 	/** Stores number of nodes per each LOD. */
 	TArray<FThreadSafeCounter> NodeCount;
 
-	/** Bounds with uniform extent - used for node generation. */
-	FBox UniformBounds;
-
-	/** Bounds covering all nodes. */
-	FBox Bounds;
-
-	/** Bounds covering all points. */
-	FBox PointsBounds;
+	/** Extent of this Cloud. */
+	FVector Extent;
 
 	/** Used to cache the Allocated Size. */
 	mutable int32 PreviousNodeCount;
 	mutable int64 PreviousPointCount;
+	mutable int64 PreviousAllocatedStructureSize;
 	mutable int64 PreviousAllocatedSize;
 
 	/** Used to notify any linked traversal octrees when they need to re-generate the data. */
-	TArray<FLidarPointCloudTraversalOctree*> LinkedTraversalOctrees;
+	TArray<TWeakPtr<FLidarPointCloudTraversalOctree, ESPMode::ThreadSafe>> LinkedTraversalOctrees;
 
 	/** Stores collision mesh data */
 	FTriMeshCollisionData CollisionMesh;
 
+	/** Pointer to the owner of this Octree */
+	ULidarPointCloud* Owner;
+
+	TQueue<FLidarPointCloudOctreeNode*> QueuedNodes;
+	TArray<FLidarPointCloudOctreeNode*> NodesInUse;
+
+	TAtomic<bool> bStreamingBusy;
+
+	/** Set to true when the Octree is persistently force-loaded. */
+	bool bIsFullyLoaded;
+
 public:
-	FLidarPointCloudOctree();
+	FLidarPointCloudOctree() : FLidarPointCloudOctree(nullptr) {}
+	FLidarPointCloudOctree(ULidarPointCloud* Owner);
 	~FLidarPointCloudOctree();
 	FLidarPointCloudOctree(const FLidarPointCloudOctree&) = delete;
 	FLidarPointCloudOctree(FLidarPointCloudOctree&&) = delete;
@@ -187,21 +222,18 @@ public:
 
 	/** Returns true if the Root node exists and has any data assigned. */
 	bool HasData() const { return Root.GetNumPoints() > 0; }
-	
+
 	/** Returns the number of different LODs. */
 	int32 GetNumLODs() const;
 
-	/** Returns the Octree bounds. */
-	FBox GetUniformBounds() const { return UniformBounds; }
-
 	/** Returns the Cloud bounds. */
-	FBox GetBounds() const { return Bounds; }
+	FBox GetBounds() const { return FBox(-Extent, Extent); }
 
-	/** Returns the Cloud bounds. */
-	FBox GetPointsBounds() const { return PointsBounds; }
+	/** Returns the extent of the Cloud's bounds. */
+	FORCEINLINE FVector GetExtent() const { return Extent; }
 
 	/** Recalculates and updates points bounds. */
-	void RefreshPointsBounds();
+	void RefreshBounds();
 
 	/** Returns the total number of points. */
 	int64 GetNumPoints() const;
@@ -209,9 +241,15 @@ public:
 	/** Returns the total number of nodes. */
 	int32 GetNumNodes() const;
 
-	/** Returns the amount of memory used by this Octree */
+	/** Returns a pointer to the Point Cloud asset, which owns this Octree. */
+	ULidarPointCloud* GetOwner() const { return Owner; }
+
+	/** Returns the amount of memory used by this Octree, including the BulkData */
 	int64 GetAllocatedSize() const;
 
+	/** Returns the amount of memory used by this Octree, excluding the BulkData */
+	int64 GetAllocatedStructureSize() const;
+
 	/** Returns the grid cell size at root level. */
 	float GetRootCellSize() const { return SharedData[0].GridSize.GetMax(); }
 
@@ -231,25 +269,25 @@ public:
 	const FTriMeshCollisionData* GetCollisionData() const { return &CollisionMesh; }
 
 	/** Populates the given array with points from the tree */
-	void GetPoints(TArray<FLidarPointCloudPoint*>& Points, int64 StartIndex = 0, int64 Count = -1) const;
+	void GetPoints(TArray<FLidarPointCloudPoint*>& Points, int64 StartIndex = 0, int64 Count = -1);
 
 	/** Populates the array with the list of points within the given sphere. */
-	void GetPointsInSphere(TArray<FLidarPointCloudPoint*>& SelectedPoints, const FSphere& Sphere, const bool& bVisibleOnly) const;
+	void GetPointsInSphere(TArray<FLidarPointCloudPoint*>& SelectedPoints, const FSphere& Sphere, const bool& bVisibleOnly);
 
 	/** Populates the array with the list of pointers to points within the given box. */
-	void GetPointsInBox(TArray<FLidarPointCloudPoint*>& SelectedPoints, const FBox& Box, const bool& bVisibleOnly) const;
+	void GetPointsInBox(TArray<FLidarPointCloudPoint*>& SelectedPoints, const FBox& Box, const bool& bVisibleOnly);
 
 	/** Populates the array with the list of points within the given frustum. */
-	void GetPointsInFrustum(TArray<FLidarPointCloudPoint*>& SelectedPoints, const FConvexVolume& Frustum, const bool& bVisibleOnly) const;
+	void GetPointsInFrustum(TArray<FLidarPointCloudPoint*>& SelectedPoints, const FConvexVolume& Frustum, const bool& bVisibleOnly);
 
 	/** Populates the given array with copies of points from the tree */
-	void GetPointsAsCopies(TArray<FLidarPointCloudPoint>& Points, int64 StartIndex = 0, int64 Count = -1) const;
+	void GetPointsAsCopies(TArray<FLidarPointCloudPoint>& Points, const FTransform* LocalToWorld, int64 StartIndex = 0, int64 Count = -1) const;
 
 	/** Populates the array with the list of points within the given sphere. */
-	void GetPointsInSphereAsCopies(TArray<FLidarPointCloudPoint>& SelectedPoints, const FSphere& Sphere, const bool& bVisibleOnly) const;
+	void GetPointsInSphereAsCopies(TArray<FLidarPointCloudPoint>& SelectedPoints, const FSphere& Sphere, const bool& bVisibleOnly, const FTransform* LocalToWorld) const;
 
 	/** Populates the array with the list of pointers to points within the given box. */
-	void GetPointsInBoxAsCopies(TArray<FLidarPointCloudPoint>& SelectedPoints, const FBox& Box, const bool& bVisibleOnly) const;
+	void GetPointsInBoxAsCopies(TArray<FLidarPointCloudPoint>& SelectedPoints, const FBox& Box, const bool& bVisibleOnly, const FTransform* LocalToWorld) const;
 
 	/** Performs a raycast test against the point cloud. Returns the pointer if hit or nullptr otherwise. */
 	FLidarPointCloudPoint* RaycastSingle(const FLidarPointCloudRay& Ray, const float& Radius, const bool& bVisibleOnly);
@@ -260,7 +298,7 @@ public:
 	 * Returns true it anything has been hit.
 	 */
 	bool RaycastMulti(const FLidarPointCloudRay& Ray, const float& Radius, const bool& bVisibleOnly, TArray<FLidarPointCloudPoint*>& OutHits);
-	bool RaycastMulti(const FLidarPointCloudRay& Ray, const float& Radius, const bool& bVisibleOnly, TArray<FLidarPointCloudPoint>& OutHits);
+	bool RaycastMulti(const FLidarPointCloudRay& Ray, const float& Radius, const bool& bVisibleOnly, const FTransform* LocalToWorld, TArray<FLidarPointCloudPoint>& OutHits);
 
 	/** Returns true if there are any points within the given sphere. */
 	bool HasPointsInSphere(const FSphere& Sphere, const bool& bVisibleOnly) const;
@@ -271,78 +309,30 @@ public:
 	/** Sets visibility of points within the given sphere. */
 	void SetVisibilityOfPointsInSphere(const bool& bNewVisibility, const FSphere& Sphere);
 
-	/**
-	 * Sets visibility of points within the given sphere.
-	 * Async version - does not wait for completion before returning from the call.
-	 */
-	void SetVisibilityOfPointsInSphereAsync(const bool& bNewVisibility, const FSphere& Sphere, TFunction<void(void)> CompletionCallback = nullptr);
-
 	/** Sets visibility of points within the given box. */
 	void SetVisibilityOfPointsInBox(const bool& bNewVisibility, const FBox& Box);
 
-	/**
-	 * Sets visibility of points within the given box.
-	 * Async version - does not wait for completion before returning from the call.
-	 */
-	void SetVisibilityOfPointsInBoxAsync(const bool& bNewVisibility, const FBox& Box, TFunction<void(void)> CompletionCallback = nullptr);
-
 	/** Sets visibility of points hit by the given ray. */
 	void SetVisibilityOfPointsByRay(const bool& bNewVisibility, const FLidarPointCloudRay& Ray, const float& Radius);
 
-	/**
-	 * Sets visibility of points hit by the given ray.
-	 * Async version - does not wait for completion before returning from the call.
-	 */
-	void SetVisibilityOfPointsByRayAsync(const bool& bNewVisibility, const FLidarPointCloudRay& Ray, const float& Radius, TFunction<void(void)> CompletionCallback = nullptr);
-
 	/** Marks all points hidden */
 	void HideAll();
 
 	/** Marks all points visible */
 	void UnhideAll();
 
-	/**
-	 * Marks all points visible
-	 * Async version - does not wait for completion before returning from the call.
-	 */
-	void ResetVisibilityAsync(TFunction<void(void)> CompletionCallback = nullptr);
-
 	/** Executes the provided action on each of the points. */
 	void ExecuteActionOnAllPoints(TFunction<void(FLidarPointCloudPoint*)> Action, const bool& bVisibleOnly);
 
-	/**
-	 * Executes the provided action on each of the points.
-	 * Async version - does not wait for completion before returning from the call.
-	 */
-	void ExecuteActionOnAllPointsAsync(TFunction<void(FLidarPointCloudPoint*)> Action, const bool& bVisibleOnly, TFunction<void(void)> CompletionCallback = nullptr);
-
 	/** Executes the provided action on each of the points within the given sphere. */
 	void ExecuteActionOnPointsInSphere(TFunction<void(FLidarPointCloudPoint*)> Action, const FSphere& Sphere, const bool& bVisibleOnly);
 
-	/**
-	 * Executes the provided action on each of the points within the given sphere.
-	 * Async version - does not wait for completion before returning from the call.
-	 */
-	void ExecuteActionOnPointsInSphereAsync(TFunction<void(FLidarPointCloudPoint*)> Action, const FSphere& Sphere, const bool& bVisibleOnly, TFunction<void(void)> CompletionCallback = nullptr);
-
 	/** Executes the provided action on each of the points within the given box. */
 	void ExecuteActionOnPointsInBox(TFunction<void(FLidarPointCloudPoint*)> Action, const FBox& Box, const bool& bVisibleOnly);
 
-	/**
-	 * Executes the provided action on each of the points within the given box.
-	 * Async version - does not wait for completion before returning from the call.
-	 */
-	void ExecuteActionOnPointsInBoxAsync(TFunction<void(FLidarPointCloudPoint*)> Action, const FBox& Box, const bool& bVisibleOnly, TFunction<void(void)> CompletionCallback = nullptr);
-
 	/** Executes the provided action on each of the points hit by the given ray. */
 	void ExecuteActionOnPointsByRay(TFunction<void(FLidarPointCloudPoint*)> Action, const FLidarPointCloudRay& Ray, const float& Radius, bool bVisibleOnly);
 
-	/**
-	 * Executes the provided action on each of the points hit by the given ray.
-	 * Async version - does not wait for completion before returning from the call.
-	 */
-	void ExecuteActionOnPointsByRayAsync(TFunction<void(FLidarPointCloudPoint*)> Action, const FLidarPointCloudRay& Ray, const float& Radius, bool bVisibleOnly, TFunction<void(void)> CompletionCallback = nullptr);
-
 	/** Applies the given color to all points */
 	void ApplyColorToAllPoints(const FColor& NewColor, const bool& bVisibleOnly);
 
@@ -362,82 +352,95 @@ public:
 	void MarkPointVisibilityDirty();
 
 	/** Initializes the Octree properties. */
-	void Initialize(const FBox& InBounds);
-
-	/** Applies given offset to all contained points. */
-	void ShiftPointsBy(FDoubleVector Offset, bool bRefreshPointsBounds);
+	void Initialize(const FVector& InExtent);
 
 	/** Inserts the given point into the Octree structure, internally thread-safe. */
-	FORCEINLINE void InsertPoint(const FLidarPointCloudPoint* Point, ELidarPointCloudDuplicateHandling DuplicateHandling, bool bRefreshPointsBounds) { InsertPoints(Point, 1, DuplicateHandling, bRefreshPointsBounds); }
+	FORCEINLINE void InsertPoint(const FLidarPointCloudPoint* Point, ELidarPointCloudDuplicateHandling DuplicateHandling, bool bRefreshPointsBounds, const FVector& Translation) { InsertPoints(Point, 1, DuplicateHandling, bRefreshPointsBounds, Translation); }
 
 	/** Inserts group of points into the Octree structure, internally thread-safe. */
 	template <typename T>
-	void InsertPoints(T Points, const int32& Count, ELidarPointCloudDuplicateHandling DuplicateHandling, bool bRefreshPointsBounds)
+	void InsertPoints(T Points, const int32& Count, ELidarPointCloudDuplicateHandling DuplicateHandling, bool bRefreshPointsBounds, const FVector& Translation)
 	{
-		Root.InsertPoints(this, Points, Count, DuplicateHandling);
+		Root.InsertPoints(this, Points, Count, DuplicateHandling, Translation);
 		MarkTraversalOctreesForInvalidation();
 		if (bRefreshPointsBounds)
 		{
-			RefreshPointsBounds();
+			RefreshBounds();
 		}
 	}
 
 	/** Attempts to remove the given point.  */
-	void RemovePoint(const FLidarPointCloudPoint* Point, bool bRefreshPointsBounds);
-	void RemovePoint(FLidarPointCloudPoint Point, bool bRefreshPointsBounds);
+	void RemovePoint(const FLidarPointCloudPoint* Point);
+	void RemovePoint(FLidarPointCloudPoint Point);
 
 	/** Removes points in bulk */
-	void RemovePoints(TArray<FLidarPointCloudPoint*>& Points, bool bRefreshPointsBounds);
+	void RemovePoints(TArray<FLidarPointCloudPoint*>& Points);
 
 	/** Removes all points within the given sphere */
-	void RemovePointsInSphere(const FSphere& Sphere, const bool& bVisibleOnly, bool bRefreshPointsBounds);
+	void RemovePointsInSphere(const FSphere& Sphere, const bool& bVisibleOnly);
 
 	/** Removes all points within the given box */
-	void RemovePointsInBox(const FBox& Box, const bool& bVisibleOnly, bool bRefreshPointsBounds);
+	void RemovePointsInBox(const FBox& Box, const bool& bVisibleOnly);
 
 	/** Removes all points hit by the given ray */
-	void RemovePointsByRay(const FLidarPointCloudRay& Ray, const float& Radius, const bool& bVisibleOnly, bool bRefreshPointsBounds);
+	void RemovePointsByRay(const FLidarPointCloudRay& Ray, const float& Radius, const bool& bVisibleOnly);
 
 	/** Removes all hidden points */
-	void RemoveHiddenPoints(bool bRefreshPointsBounds);
+	void RemoveHiddenPoints();
 
 	/** Removes all points and, optionally, all nodes except for the root node. Retains the bounds. */
 	void Empty(bool bDestroyNodes);
 
 	/** Adds the given traversal octree to the list of linked octrees. */
-	void RegisterTraversalOctree(FLidarPointCloudTraversalOctree* TraversalOctree)
+	void RegisterTraversalOctree(TWeakPtr<FLidarPointCloudTraversalOctree, ESPMode::ThreadSafe> TraversalOctree)
 	{
-		if (TraversalOctree)
+		if (TraversalOctree.IsValid())
 		{
 			LinkedTraversalOctrees.Add(TraversalOctree);
 		}
 	}
 
 	/** Removes the given traversal octree from the list */
-	void UnregisterTraversalOctree(FLidarPointCloudTraversalOctree* TraversalOctree)
-	{
-		if (TraversalOctree)
-		{
-			LinkedTraversalOctrees.Remove(TraversalOctree);
-		}
-	}
+	void UnregisterTraversalOctree(FLidarPointCloudTraversalOctree* TraversalOctree);
+
+	/** If bImmediate is true, the node will be loaded immediately, otherwise, it will be queued for async loading. */
+	void QueueNode(FLidarPointCloudOctreeNode* Node, float Lifetime);
+
+	/** Streams all requested nodes */
+	void StreamQueuedNodes();
+
+	void UnloadOldNodes(const float& CurrentTime);
+
+	/** Returns true, if the cloud is fully and persistently loaded. */
+	bool IsFullyLoaded() const { return bIsFullyLoaded; }
+
+	/** Persistently loads all nodes. */
+	void LoadAllNodes();
 
-private:
 	/**
-	 * Applies new bounds.
-	 * Returns true if successful
+	 * Releases all nodes.
+	 * Optionally, releases persistent nodes too.
 	 */
-	bool SetNewBounds(const FBox& InBounds);
-	
-	void RemovePoint_Internal(FLidarPointCloudOctreeNode* Node, int32 Index, bool bAllocatedPoint);
+	void ReleaseAllNodes(bool bIncludePersistent);
+
+private:
+	void RefreshAllocatedSize();
+
+	void RemovePoint_Internal(FLidarPointCloudOctreeNode* Node, int32 Index);
 
 	/** Notifies all linked traversal octrees that they should invalidate and regenerate the data. */
 	void MarkTraversalOctreesForInvalidation();
 
-	friend FArchive& operator<<(FArchive& Ar, FLidarPointCloudOctree& O);
+	void Serialize(FArchive& Ar);
+
+	friend FArchive& operator<<(FArchive& Ar, FLidarPointCloudOctree& O)
+	{
+		O.Serialize(Ar);
+		return Ar;
+	}
 
 	friend FLidarPointCloudOctreeNode;
-	friend struct FLidarPointCloudTraversalOctree;
+	friend FLidarPointCloudTraversalOctree;
 };
 
 /**
@@ -457,7 +460,6 @@ struct FLidarPointCloudTraversalOctreeNode
 	/** Calculated for use with adaptive sprite scaling */
 	uint8 VirtualDepth;
 
-	// #todo: Can be included as a first item in the Children array to save 8 bytes
 	FLidarPointCloudTraversalOctreeNode* Parent;
 
 	/** Stores the children array */
@@ -469,7 +471,7 @@ struct FLidarPointCloudTraversalOctreeNode
 	FLidarPointCloudTraversalOctreeNode();
 
 	/** Builds the traversal version of the given node. */
-	void Build(FLidarPointCloudOctreeNode* Node, const FTransform& LocalToWorld);
+	void Build(FLidarPointCloudOctreeNode* Node, const FTransform& LocalToWorld, const FVector& LocationOffset);
 
 	// #refactor: This is ugly - refactor to avoid the const or call directly in the buffer update with no mutable
 	/** Calculates virtual depth of this node, to be used to estimate the best sprite size */
@@ -490,11 +492,11 @@ struct FLidarPointCloudTraversalOctree
 	/** Stores the number of LODs. */
 	uint8 NumLODs;
 
-	/** Stores number of points per each LOD. */
-	TArray<int64> PointCount;
+	/** Normalized histogram of level weights, one for each LOD. Used for point scaling */
+	TArray<float> LevelWeights;
 
-	/** Total number of points. */
-	int64 NumPoints;
+	float VirtualDepthMultiplier;
+	float ReversedVirtualDepthMultiplier;
 
 	/** Pointer to the source Octree */
 	FLidarPointCloudOctree* Octree;
@@ -510,11 +512,8 @@ struct FLidarPointCloudTraversalOctree
 	FLidarPointCloudTraversalOctree& operator=(const FLidarPointCloudTraversalOctree&) = delete;
 	FLidarPointCloudTraversalOctree& operator=(FLidarPointCloudTraversalOctree&&) = delete;
 
-	/** Selects and returns the subset of visible nodes for rendering. */
-	TArray<FLidarPointCloudTraversalOctreeNode*> GetVisibleNodes(const class FSceneView* View, const float& MinScreenSizeSq, const int32& PointBudget, const float& ScreenCenterImportance, const int32& MinDepth, const int32& MaxDepth);
-	
-	void CalculateSpriteSize(TArray<FLidarPointCloudTraversalOctreeNode*>& SelectedNodes, const float& PointSizeBias);
+	/** Selects and appends the subset of visible nodes for rendering. */
+	void GetVisibleNodes(TArray<FLidarPointCloudLODManager::FNodeSizeData>& NodeSizeData, const FLidarPointCloudViewData& ViewData, const int32& ProxyIndex, const FLidarPointCloudNodeSelectionParams& SelectionParams, const float& CurrentTime);
 
-	/** Returns normalized histogram of level weights, one for each LOD. Used for point scaling */
-	TArray<float> CalculateLevelWeights() const;
+	FVector GetCenter() const { return Root.Center; }
 };
diff --git a/Engine/Plugins/Enterprise/LidarPointCloud/Source/LidarPointCloudRuntime/Public/LidarPointCloudSettings.h b/Engine/Plugins/Enterprise/LidarPointCloud/Source/LidarPointCloudRuntime/Public/LidarPointCloudSettings.h
--- a/Engine/Plugins/Enterprise/LidarPointCloud/Source/LidarPointCloudRuntime/Public/LidarPointCloudSettings.h
+++ b/Engine/Plugins/Enterprise/LidarPointCloud/Source/LidarPointCloudRuntime/Public/LidarPointCloudSettings.h
@@ -65,6 +65,21 @@ public:
 	UPROPERTY(config, EditAnywhere, Category=Performance)
 	int32 ExportBatchSize;
 
+	/**
+	 * Enabling this will allocate larger portion of the available point budget to the viewport with focus.
+	 * May improve asset editing experience, if the scenes are busy.
+	 * Disable, if you are experiencing visual glitches.
+	 */
+	UPROPERTY(config, EditAnywhere, Category=Performance)
+	bool bPrioritizeActiveViewport;
+
+	/**
+	 * Enabling this will compress data when saving the assets.
+	 * May introduce delay when streaming points on slower machines.
+	 */
+	UPROPERTY(config, EditAnywhere, Category=IO)
+	bool bUseCompression;
+
 	/** Affects the size of per-thread data for the meshing algorithm. */
 	UPROPERTY(config, EditAnywhere, Category=Collision)
 	int32 MeshingBatchSize;
diff --git a/Engine/Plugins/Enterprise/LidarPointCloud/Source/LidarPointCloudRuntime/Public/LidarPointCloudShared.h b/Engine/Plugins/Enterprise/LidarPointCloud/Source/LidarPointCloudRuntime/Public/LidarPointCloudShared.h
--- a/Engine/Plugins/Enterprise/LidarPointCloud/Source/LidarPointCloudRuntime/Public/LidarPointCloudShared.h
+++ b/Engine/Plugins/Enterprise/LidarPointCloud/Source/LidarPointCloudRuntime/Public/LidarPointCloudShared.h
@@ -3,6 +3,7 @@
 #pragma once
 
 #include "CoreMinimal.h"
+#include "Serialization/BulkData.h"
 #include "LidarPointCloudShared.generated.h"
 
 DECLARE_LOG_CATEGORY_EXTERN(LogLidarPointCloud, Log, All);
@@ -158,6 +159,12 @@ struct LIDARPOINTCLOUDRUNTIME_API FDoubleVector
 
 	FORCEINLINE FVector ToVector() const { return FVector(X, Y, Z); }
 	FORCEINLINE FIntVector ToIntVector() const { return FIntVector(X, Y, Z); }
+
+	friend FArchive& operator<<(FArchive& Ar, FDoubleVector& V)
+	{
+		Ar << V.X << V.Y << V.Z;
+		return Ar;
+	}
 };
 
 #pragma pack(push)
@@ -180,8 +187,9 @@ public:
 	/** Valid range is 0 - 31. */
 	uint8 ClassificationID : 5;
 
-private:
 	uint8 bSelected : 1;
+
+private:	
 	uint8 bMarkedForDeletion : 1;
 
 public:
@@ -242,17 +250,115 @@ public:
 		ClassificationID = Other.ClassificationID;
 	}
 
+	FORCEINLINE FLidarPointCloudPoint Transform(const FTransform& Transform) const
+	{
+		return FLidarPointCloudPoint(Transform.TransformPosition(Location), Color, bVisible, ClassificationID);
+	}
+
 	bool operator==(const FLidarPointCloudPoint& P) const { return Location == P.Location && Color == P.Color && bVisible == P.bVisible && ClassificationID == P.ClassificationID; }
 
 	friend FArchive& operator<<(FArchive& Ar, FLidarPointCloudPoint& P);
 	friend class FLidarPointCloudOctree;
-	friend class FLidarPointCloudInstanceBuffer;
 #if WITH_EDITOR
 	friend class FLidarPointCloudEditor;
 #endif
 };
 #pragma pack(pop)
 
+struct FLidarPointCloudBulkData : public FUntypedBulkData
+{
+private:
+	FLidarPointCloudPoint* DataPtr;
+	TAtomic<bool> bHasData;
+
+public:
+	FLidarPointCloudBulkData()
+		: DataPtr(nullptr)
+		, bHasData(false)
+	{
+	}
+
+	virtual int32 GetElementSize() const override
+	{
+		return sizeof(FLidarPointCloudPoint);
+	}
+	
+	virtual void SerializeElement(FArchive& Ar, void* Data, int64 ElementIndex) override
+	{
+		Ar.Serialize((FLidarPointCloudPoint*)Data + ElementIndex, sizeof(FLidarPointCloudPoint));
+		//Ar << *((FLidarPointCloudPoint*)Data + ElementIndex);
+	}
+	
+	virtual bool RequiresSingleElementSerialization(FArchive& Ar) override
+	{
+		return false;
+	}
+
+	/** Serves as a workaround for UnloadBulkData being editor-only */
+	void ReleaseData()
+	{
+		if (bHasData)
+		{
+			bHasData = false;
+			void* Temp = nullptr;
+			GetCopy(&Temp);
+			FMemory::Free(Temp);
+		}
+	}
+
+	FORCEINLINE FLidarPointCloudPoint* GetData() const
+	{
+		MakeSureDataIsLoaded();
+		return DataPtr;
+	}
+
+	void CopyToArray(TArray<FLidarPointCloudPoint>& Array)
+	{
+		Array.AddUninitialized(GetElementCount());
+		FMemory::Memcpy(Array.GetData(), GetData(), sizeof(FLidarPointCloudPoint) * Array.Num());
+	}
+
+	void CopyFromArray(TArray<FLidarPointCloudPoint>& Array)
+	{
+		Lock(LOCK_READ_WRITE);
+		DataPtr = (FLidarPointCloudPoint*)Realloc(Array.Num());
+		FMemory::Memcpy(DataPtr, Array.GetData(), Array.Num() * sizeof(FLidarPointCloudPoint));
+		bHasData = true;
+		Unlock();
+	}
+
+	/** Serializes data from legacy arrays */
+	void SerializeLegacy(FArchive& Ar)
+	{
+		// Load legacy data
+		TArray<FLidarPointCloudPoint> AllocatedPoints;
+		TArray<FLidarPointCloudPoint> PaddingPoints;
+		Ar << AllocatedPoints << PaddingPoints;
+
+		// Copy the data to BulkData
+		Lock(LOCK_READ_WRITE);
+		DataPtr = (FLidarPointCloudPoint*)Realloc(AllocatedPoints.Num() + PaddingPoints.Num());
+		FMemory::Memcpy(DataPtr, AllocatedPoints.GetData(), AllocatedPoints.Num() * sizeof(FLidarPointCloudPoint));
+		FMemory::Memcpy(DataPtr + AllocatedPoints.Num(), PaddingPoints.GetData(), PaddingPoints.Num() * sizeof(FLidarPointCloudPoint));
+		bHasData = true;
+		Unlock();
+	}
+
+	FORCEINLINE bool HasData() const { return bHasData; }
+
+private:
+	void MakeSureDataIsLoaded() const
+	{
+		if (!bHasData)
+		{
+			FLidarPointCloudBulkData* mutable_this = const_cast<FLidarPointCloudBulkData*>(this);
+			mutable_this->DataPtr = (FLidarPointCloudPoint*)LockReadOnly();
+			mutable_this->bHasData = true;
+			Unlock();
+		}
+	}
+};
+
 /** Used in blueprint latent function execution */
 UENUM(BlueprintType)
 enum class ELidarPointCloudAsyncMode : uint8
@@ -262,6 +368,63 @@ enum class ELidarPointCloudAsyncMode : uint8
 	Progress
 };
 
+/** Used to help track multiple buffer allocations */
+class FLidarPointCloudDataBuffer
+{
+public:
+	FLidarPointCloudDataBuffer() : bInUse(false), PendingSize(0) {}
+	~FLidarPointCloudDataBuffer() = default;
+	FLidarPointCloudDataBuffer(const FLidarPointCloudDataBuffer& Other)
+	{
+		Data = Other.Data;
+		bInUse = false;
+	}
+	FLidarPointCloudDataBuffer(FLidarPointCloudDataBuffer&&) = delete;
+	FLidarPointCloudDataBuffer& operator=(const FLidarPointCloudDataBuffer& Other)
+	{
+		Data = Other.Data;
+		bInUse = false;
+		return *this;
+	}
+	FLidarPointCloudDataBuffer& operator=(FLidarPointCloudDataBuffer&&) = delete;
+
+	FORCEINLINE uint8* GetData() { return Data.GetData(); }
+	FORCEINLINE bool InUse() const { return bInUse; }
+	
+	/** Marks the buffer as no longer in use so it can be reassigned to another read thread. */
+	void MarkAsFree();
+	void Initialize(const int32& Size);
+	void Resize(const int32& NewBufferSize, bool bForce = false);
+
+private:
+	TAtomic<bool> bInUse;
+	TArray<uint8> Data;
+	int32 PendingSize;
+
+	friend class FLidarPointCloudDataBufferManager;
+};
+
+/** Used to help track multiple buffer allocations */
+class FLidarPointCloudDataBufferManager
+{
+public:
+	FLidarPointCloudDataBufferManager(const int32& BufferSize);
+	~FLidarPointCloudDataBufferManager();
+	FLidarPointCloudDataBufferManager(const FLidarPointCloudDataBufferManager&) = delete;
+	FLidarPointCloudDataBufferManager(FLidarPointCloudDataBufferManager&&) = delete;
+	FLidarPointCloudDataBufferManager& operator=(const FLidarPointCloudDataBufferManager&) = delete;
+	FLidarPointCloudDataBufferManager& operator=(FLidarPointCloudDataBufferManager&&) = delete;
+
+	FLidarPointCloudDataBuffer* GetFreeBuffer();
+
+	void Resize(const int32& NewBufferSize);
+
+private:
+	int32 BufferSize;
+	TList<FLidarPointCloudDataBuffer> Head;
+	TList<FLidarPointCloudDataBuffer>* Tail;
+};
+
 /** Used for Raycasting */
 struct FLidarPointCloudRay
 {
@@ -285,6 +448,10 @@ public:
 		SetDirection(Transform.TransformVector(Direction));
 		return *this;
 	}
+	FORCEINLINE FLidarPointCloudRay ShiftBy(const FVector& Offset) const
+	{
+		return FLidarPointCloudRay(Origin + Offset, Direction);
+	}
 
 	FORCEINLINE FVector GetDirection() const { return Direction; }
 	FORCEINLINE void SetDirection(const FVector& NewDirection)
