The usage of the auto keyword in the Lidar plugin was against our coding standard.
https://docs.unrealengine.com/en-US/Programming/Development/CodingStandard/#the'auto'keyword

diff --git a/Engine/Plugins/Enterprise/LidarPointCloud/Source/LidarPointCloudEditor/Private/LidarPointCloudEditor.cpp b/Engine/Plugins/Enterprise/LidarPointCloud/Source/LidarPointCloudEditor/Private/LidarPointCloudEditor.cpp
--- a/Engine/Plugins/Enterprise/LidarPointCloud/Source/LidarPointCloudEditor/Private/LidarPointCloudEditor.cpp
+++ b/Engine/Plugins/Enterprise/LidarPointCloud/Source/LidarPointCloudEditor/Private/LidarPointCloudEditor.cpp
@@ -90,7 +90,7 @@ FLidarPointCloudEditor::~FLidarPointCloudEditor()
 void FLidarPointCloudEditor::RegisterTabSpawners(const TSharedRef<class FTabManager>& InTabManager)
 {
 	WorkspaceMenuCategory = InTabManager->AddLocalWorkspaceMenuCategory(LOCTEXT("WorkspaceMenu_LidarPointCloudEditor", "LiDAR Point Cloud Editor"));
-	auto WorkspaceMenuCategoryRef = WorkspaceMenuCategory.ToSharedRef();
+	TSharedRef<FWorkspaceItem> WorkspaceMenuCategoryRef = WorkspaceMenuCategory.ToSharedRef();
 
 	FAssetEditorToolkit::RegisterTabSpawners(InTabManager);
 
@@ -449,7 +449,7 @@ void FLidarPointCloudEditor::ExtendToolBar()
 	{
 		static void FillToolbar(FToolBarBuilder& ToolbarBuilder, FLidarPointCloudEditor* ThisEditor)
 		{
-			auto Commands = &FLidarPointCloudEditorCommands::Get();
+			const FLidarPointCloudEditorCommands* Commands = &FLidarPointCloudEditorCommands::Get();
 
 			ToolbarBuilder.BeginSection("Camera");
 			{
diff --git a/Engine/Plugins/Enterprise/LidarPointCloud/Source/LidarPointCloudEditor/Private/LidarPointCloudEditorViewport.cpp b/Engine/Plugins/Enterprise/LidarPointCloud/Source/LidarPointCloudEditor/Private/LidarPointCloudEditorViewport.cpp
--- a/Engine/Plugins/Enterprise/LidarPointCloud/Source/LidarPointCloudEditor/Private/LidarPointCloudEditorViewport.cpp
+++ b/Engine/Plugins/Enterprise/LidarPointCloud/Source/LidarPointCloudEditor/Private/LidarPointCloudEditorViewport.cpp
@@ -36,7 +36,7 @@ public:
 		const bool bInShouldCloseWindowAfterMenuSelection = true;
 		FMenuBuilder ShowMenuBuilder(bInShouldCloseWindowAfterMenuSelection, ViewportRef->GetCommandList());
 		{
-			auto Commands = FLidarPointCloudEditorCommands::Get();
+			const FLidarPointCloudEditorCommands& Commands = FLidarPointCloudEditorCommands::Get();
 
 			ShowMenuBuilder.AddMenuEntry(Commands.SetShowGrid);
 			ShowMenuBuilder.AddMenuEntry(Commands.SetShowBounds);
@@ -80,8 +80,6 @@ void SLidarPointCloudEditorViewport::Construct(const FArguments& InArgs)
 		[
 			SAssignNew(OverlayTextVerticalBox, SVerticalBox)
 		];
-
-	FCoreUObjectDelegates::OnObjectPropertyChanged.AddRaw(this, &SLidarPointCloudEditorViewport::OnObjectPropertyChanged);
 }
 
 SLidarPointCloudEditorViewport::SLidarPointCloudEditorViewport()
@@ -91,7 +89,6 @@ SLidarPointCloudEditorViewport::SLidarPointCloudEditorViewport()
 
 SLidarPointCloudEditorViewport::~SLidarPointCloudEditorViewport()
 {
-	FCoreUObjectDelegates::OnObjectPropertyChanged.RemoveAll(this);
 	if (EditorViewportClient.IsValid())
 	{
 		EditorViewportClient->Viewport = nullptr;
@@ -119,7 +116,7 @@ void SLidarPointCloudEditorViewport::RefreshViewport()
 
 void SLidarPointCloudEditorViewport::ResetCamera()
 {
-	if (auto EditorViewportClientRawPtr = EditorViewportClient.Get())
+	if (FLidarPointCloudEditorViewportClient* EditorViewportClientRawPtr = EditorViewportClient.Get())
 	{
 		EditorViewportClientRawPtr->ResetCamera();
 	}
@@ -139,7 +136,7 @@ void SLidarPointCloudEditorViewport::PopulateOverlayText(const TArray<FOverlayTe
 {
 	OverlayTextVerticalBox->ClearChildren();
 
-	for (const auto& TextItem : TextItems)
+	for (const FOverlayTextItem& TextItem : TextItems)
 	{
 		OverlayTextVerticalBox->AddSlot()
 			[
@@ -155,18 +152,6 @@ bool SLidarPointCloudEditorViewport::IsVisible() const
 	return ViewportWidget.IsValid() && (!ParentTab.IsValid() || ParentTab.Pin()->IsForeground());
 }
 
-void SLidarPointCloudEditorViewport::OnObjectPropertyChanged(UObject* ObjectBeingModified, FPropertyChangedEvent& PropertyChangedEvent)
-{
-	if (!ensure(ObjectBeingModified))
-	{
-		return;
-	}
-
-	if (PreviewCloudComponent)
-	{
-	}
-}
-
 TSharedRef<FEditorViewportClient> SLidarPointCloudEditorViewport::MakeEditorViewportClient()
 {
 	EditorViewportClient = MakeShareable(new FLidarPointCloudEditorViewportClient(PointCloudEditorPtr, SharedThis(this), &PreviewScene, PointCloud, PreviewCloudComponent));
diff --git a/Engine/Plugins/Enterprise/LidarPointCloud/Source/LidarPointCloudEditor/Private/LidarPointCloudEditorViewport.h b/Engine/Plugins/Enterprise/LidarPointCloud/Source/LidarPointCloudEditor/Private/LidarPointCloudEditorViewport.h
--- a/Engine/Plugins/Enterprise/LidarPointCloud/Source/LidarPointCloudEditor/Private/LidarPointCloudEditorViewport.h
+++ b/Engine/Plugins/Enterprise/LidarPointCloud/Source/LidarPointCloudEditor/Private/LidarPointCloudEditorViewport.h
@@ -84,9 +84,6 @@ private:
 	/** Determines the visibility of the viewport. */
 	bool IsVisible() const override;
 
-	/** Callback for updating preview socket meshes if the static mesh or socket has been modified. */
-	void OnObjectPropertyChanged(UObject* ObjectBeingModified, FPropertyChangedEvent& PropertyChangedEvent);
-
 private:
 	/** The parent tab where this viewport resides */
 	TWeakPtr<SDockTab> ParentTab;
diff --git a/Engine/Plugins/Enterprise/LidarPointCloud/Source/LidarPointCloudEditor/Private/LidarPointCloudEditorViewportClient.cpp b/Engine/Plugins/Enterprise/LidarPointCloud/Source/LidarPointCloudEditor/Private/LidarPointCloudEditorViewportClient.cpp
--- a/Engine/Plugins/Enterprise/LidarPointCloud/Source/LidarPointCloudEditor/Private/LidarPointCloudEditorViewportClient.cpp
+++ b/Engine/Plugins/Enterprise/LidarPointCloud/Source/LidarPointCloudEditor/Private/LidarPointCloudEditorViewportClient.cpp
@@ -229,8 +229,8 @@ void FLidarPointCloudEditorViewportClient::Draw(const FSceneView* View, FPrimiti
 
 void FLidarPointCloudEditorViewportClient::DrawCanvas(FViewport& InViewport, FSceneView& View, FCanvas& Canvas)
 {
-	auto PointCloudEditor = PointCloudEditorPtr.Pin();
-	auto PointCloudEditorViewport = PointCloudEditorViewportPtr.Pin();
+	TSharedPtr<FLidarPointCloudEditor> PointCloudEditor = PointCloudEditorPtr.Pin();
+	TSharedPtr<SLidarPointCloudEditorViewport> PointCloudEditorViewport = PointCloudEditorViewportPtr.Pin();
 	if (!PointCloudEditor.IsValid() || !PointCloudEditorViewport.IsValid())
 	{
 		return;
@@ -333,7 +333,7 @@ void FLidarPointCloudEditorViewportClient::OnSelectionEnd(const FIntVector4& Sel
 {
 	if (ULidarPointCloudComponent* PointCloudComponentRawPtr = PointCloudComponent.Get())
 	{
-		if (auto Editor = PointCloudEditorPtr.Pin())
+		if (TSharedPtr<FLidarPointCloudEditor> Editor = PointCloudEditorPtr.Pin())
 		{
 			// Compute a view.
 			FSceneViewFamilyContext ViewFamily(FSceneViewFamily::ConstructionValues(Viewport, GetScene(), EngineShowFlags).SetRealtimeUpdate(IsRealtime()));
diff --git a/Engine/Plugins/Enterprise/LidarPointCloud/Source/LidarPointCloudEditor/Private/LidarPointCloudFactory.cpp b/Engine/Plugins/Enterprise/LidarPointCloud/Source/LidarPointCloudEditor/Private/LidarPointCloudFactory.cpp
--- a/Engine/Plugins/Enterprise/LidarPointCloud/Source/LidarPointCloudEditor/Private/LidarPointCloudFactory.cpp
+++ b/Engine/Plugins/Enterprise/LidarPointCloud/Source/LidarPointCloudEditor/Private/LidarPointCloudFactory.cpp
@@ -21,14 +21,14 @@ FText FAssetTypeActions_LidarPointCloud::GetName() const
 void FAssetTypeActions_LidarPointCloud::GetActions(const TArray<UObject*>& InObjects, FMenuBuilder& MenuBuilder)
 {
 	TArray<ULidarPointCloud*> PointClouds;
-	for (auto ObjIt = InObjects.CreateConstIterator(); ObjIt; ++ObjIt)
+	for (UObject* Object : InObjects)
 	{
-		PointClouds.Add(CastChecked<ULidarPointCloud>(*ObjIt));
+		PointClouds.Add(CastChecked<ULidarPointCloud>(Object));
 	}
 
 	// Make sure at least one asset has source assigned
 	bool bSourceExists = false;
-	for (auto& PointCloud : PointClouds)
+	for (ULidarPointCloud* PointCloud : PointClouds)
 	{
 		if (!PointCloud->GetSourcePath().IsEmpty())
 		{
@@ -78,9 +78,9 @@ void FAssetTypeActions_LidarPointCloud::OpenAssetEditor(const TArray<UObject*>&
 {
 	const EToolkitMode::Type Mode = EditWithinLevelEditor.IsValid() ? EToolkitMode::WorldCentric : EToolkitMode::Standalone;
 
-	for (auto ObjIt = InObjects.CreateConstIterator(); ObjIt; ++ObjIt)
+	for (UObject* Object : InObjects)
 	{
-		if (ULidarPointCloud* PointCloud = Cast<ULidarPointCloud>(*ObjIt))
+		if (ULidarPointCloud* PointCloud = Cast<ULidarPointCloud>(Object))
 		{
 			TSharedRef<FLidarPointCloudEditor> NewPointCloudEditor(new FLidarPointCloudEditor());
 			NewPointCloudEditor->InitPointCloudEditor(Mode, EditWithinLevelEditor, PointCloud);
diff --git a/Engine/Plugins/Enterprise/LidarPointCloud/Source/LidarPointCloudRuntime/Private/Collision/LidarPointCloudCollision.cpp b/Engine/Plugins/Enterprise/LidarPointCloud/Source/LidarPointCloudRuntime/Private/Collision/LidarPointCloudCollision.cpp
--- a/Engine/Plugins/Enterprise/LidarPointCloud/Source/LidarPointCloudRuntime/Private/Collision/LidarPointCloudCollision.cpp
+++ b/Engine/Plugins/Enterprise/LidarPointCloud/Source/LidarPointCloudRuntime/Private/Collision/LidarPointCloudCollision.cpp
@@ -458,7 +458,7 @@ void LidarPointCloudCollision::BuildCollisionMesh(FLidarPointCloudOctree* Octree
 				FMemory::Memzero(VoxelizedGrid, NumCellsCu);
 
 				// Calculate voxelized grids
-				for (auto Point = Selection.GetData(), DataEnd = Selection.GetData() + Selection.Num(); Point != DataEnd; ++Point)
+				for (FLidarPointCloudPoint** Point = Selection.GetData(), ** DataEnd = Selection.GetData() + Selection.Num(); Point != DataEnd; ++Point)
 				{
 					// Calculate location relative to sampling bounds
 					FVector Location = (*Point)->Location - SamplingBounds.Min;
@@ -482,7 +482,10 @@ void LidarPointCloudCollision::BuildCollisionMesh(FLidarPointCloudOctree* Octree
 	}
 
 	// Sync threads
-	for (auto& ThreadResult : ThreadResults) ThreadResult.Get();
+	for (const TFuture<void>& ThreadResult : ThreadResults)
+	{
+		ThreadResult.Get();
+	}
 
 	FBenchmarkTimer::Log("Collision: Meshing");
 
@@ -535,7 +538,7 @@ void LidarPointCloudCollision::BuildCollisionMesh(FLidarPointCloudOctree* Octree
 		CollisionMesh->Indices.Empty(NumVertices.GetValue() / 3);
 		CollisionMesh->Indices.AddUninitialized(NumVertices.GetValue() / 3);
 		int32 i = 0;
-		for (auto Index = (int32*)CollisionMesh->Indices.GetData(), DataEnd = Index + CollisionMesh->Indices.Num() * 3; Index != DataEnd; ++Index)
+		for (int32* Index = (int32*)CollisionMesh->Indices.GetData(), * DataEnd = Index + CollisionMesh->Indices.Num() * 3; Index != DataEnd; ++Index)
 		{
 			*Index = i++;
 		}
diff --git a/Engine/Plugins/Enterprise/LidarPointCloud/Source/LidarPointCloudRuntime/Private/IO/LidarPointCloudFileIO.cpp b/Engine/Plugins/Enterprise/LidarPointCloud/Source/LidarPointCloudRuntime/Private/IO/LidarPointCloudFileIO.cpp
--- a/Engine/Plugins/Enterprise/LidarPointCloud/Source/LidarPointCloudRuntime/Private/IO/LidarPointCloudFileIO.cpp
+++ b/Engine/Plugins/Enterprise/LidarPointCloud/Source/LidarPointCloudRuntime/Private/IO/LidarPointCloudFileIO.cpp
@@ -41,7 +41,7 @@ bool FLidarPointCloudFileIOHandler::ValidateImportSettings(TSharedPtr<FLidarPoin
 		if (ImportSettings->IsGeneric())
 		{
 			// Convert to the specialized settings
-			auto NewSettings = GetImportSettings(ImportSettings->Filename);
+			TSharedPtr<FLidarPointCloudImportSettings> NewSettings = GetImportSettings(ImportSettings->Filename);
 			NewSettings->bImportAll = ImportSettings->bImportAll;
 			ImportSettings = NewSettings;
 		}
@@ -180,7 +180,7 @@ TArray<FString> ULidarPointCloudFileIO::GetSupportedImportExtensions()
 {
 	TArray<FString> Extensions;
 
-	for (auto& Handler : Instance->RegisteredHandlers)
+	for (const TPair<FString, FLidarPointCloudFileIOHandler*>& Handler : Instance->RegisteredHandlers)
 	{
 		if (Handler.Value->SupportsImport())
 		{
@@ -195,7 +195,7 @@ TArray<FString> ULidarPointCloudFileIO::GetSupportedExportExtensions()
 {
 	TArray<FString> Extensions;
 
-	for (auto& Handler : Instance->RegisteredHandlers)
+	for (const TPair<FString, FLidarPointCloudFileIOHandler*>& Handler : Instance->RegisteredHandlers)
 	{
 		if (Handler.Value->SupportsExport())
 		{
diff --git a/Engine/Plugins/Enterprise/LidarPointCloud/Source/LidarPointCloudRuntime/Private/IO/LidarPointCloudFileIO_LAS.cpp b/Engine/Plugins/Enterprise/LidarPointCloud/Source/LidarPointCloudRuntime/Private/IO/LidarPointCloudFileIO_LAS.cpp
--- a/Engine/Plugins/Enterprise/LidarPointCloud/Source/LidarPointCloudRuntime/Private/IO/LidarPointCloudFileIO_LAS.cpp
+++ b/Engine/Plugins/Enterprise/LidarPointCloud/Source/LidarPointCloudRuntime/Private/IO/LidarPointCloudFileIO_LAS.cpp
@@ -303,7 +303,10 @@ bool ULidarPointCloudFileIO_LAS::HandleImport(const FString& Filename, TSharedPt
 				}
 
 				// Sync threads
-				for (auto& ThreadResult : ThreadResults) ThreadResult.Get();
+				for (const TFuture<void>& ThreadResult : ThreadResults)
+				{
+					ThreadResult.Get();
+				}
 
 				// Make sure to progress the counter to the end before returning
 				OutImportResults.IncrementProgressCounter(TotalPointsToRead);
diff --git a/Engine/Plugins/Enterprise/LidarPointCloud/Source/LidarPointCloudRuntime/Private/LidarPointCloud.cpp b/Engine/Plugins/Enterprise/LidarPointCloud/Source/LidarPointCloudRuntime/Private/LidarPointCloud.cpp
--- a/Engine/Plugins/Enterprise/LidarPointCloud/Source/LidarPointCloudRuntime/Private/LidarPointCloud.cpp
+++ b/Engine/Plugins/Enterprise/LidarPointCloud/Source/LidarPointCloudRuntime/Private/LidarPointCloud.cpp
@@ -672,7 +672,7 @@ bool ULidarPointCloud::InsertPoints_NoLock(T InPoints, const int64& Count, ELida
 		}
 
 		// Sync
-		for (auto& ThreadResult : ThreadResults)
+		for (const TFuture<void>& ThreadResult : ThreadResults)
 		{
 			ThreadResult.Get();
 		}
@@ -1069,9 +1069,9 @@ bool ULidarPointCloudBlueprintLibrary::LineTraceSingle(UObject* WorldContextObje
 
 		for (TActorIterator<ALidarPointCloudActor> Itr(World); Itr; ++Itr)
 		{
-			auto Actor = *Itr;
-			auto Component = Actor->GetPointCloudComponent();
-			auto Point = Component->LineTraceSingle(Ray, Radius, bVisibleOnly);
+			ALidarPointCloudActor* Actor = *Itr;
+			ULidarPointCloudComponent* Component = Actor->GetPointCloudComponent();
+			FLidarPointCloudPoint* Point = Component->LineTraceSingle(Ray, Radius, bVisibleOnly);
 			if (Point)
 			{
 				Hit = FLidarPointCloudTraceHit(Actor, Component);
@@ -1094,8 +1094,8 @@ bool ULidarPointCloudBlueprintLibrary::LineTraceMulti(UObject* WorldContextObjec
 
 		for (TActorIterator<ALidarPointCloudActor> Itr(World); Itr; ++Itr)
 		{
-			auto Actor = *Itr;
-			auto Component = Actor->GetPointCloudComponent();
+			ALidarPointCloudActor* Actor = *Itr;
+			ULidarPointCloudComponent* Component = Actor->GetPointCloudComponent();
 
 			FLidarPointCloudTraceHit Hit(Actor, Component);
 			if (Component->LineTraceMulti(Ray, Radius, bVisibleOnly, true, Hit.Points))
diff --git a/Engine/Plugins/Enterprise/LidarPointCloud/Source/LidarPointCloudRuntime/Private/LidarPointCloudLODManager.cpp b/Engine/Plugins/Enterprise/LidarPointCloud/Source/LidarPointCloudRuntime/Private/LidarPointCloudLODManager.cpp
--- a/Engine/Plugins/Enterprise/LidarPointCloud/Source/LidarPointCloudRuntime/Private/LidarPointCloudLODManager.cpp
+++ b/Engine/Plugins/Enterprise/LidarPointCloud/Source/LidarPointCloudRuntime/Private/LidarPointCloudLODManager.cpp
@@ -67,7 +67,7 @@ uint32 SetLocationAndColorData(uint8* Buffer, const TArray<FLidarPointCloudTrave
 	{
 		if (bUseClassification)
 		{
-			for (auto& Node : Nodes)
+			for (const FLidarPointCloudTraversalOctreeNode* Node : Nodes)
 			{
 				// Skip nodes with no available data
 				if (!Node->DataNode->HasData())
@@ -96,7 +96,7 @@ uint32 SetLocationAndColorData(uint8* Buffer, const TArray<FLidarPointCloudTrave
 		}
 		else
 		{
-			for (auto& Node : Nodes)
+			for (const FLidarPointCloudTraversalOctreeNode* Node : Nodes)
 			{
 				// Skip nodes with no available data
 				if (!Node->DataNode->HasData())
@@ -127,7 +127,7 @@ uint32 SetLocationAndColorData(uint8* Buffer, const TArray<FLidarPointCloudTrave
 	{
 		if (bUseClassification)
 		{
-			for (auto& Node : Nodes)
+			for (const FLidarPointCloudTraversalOctreeNode* Node : Nodes)
 			{
 				// Skip nodes with no available data
 				if (!Node->DataNode->HasData())
@@ -148,7 +148,7 @@ uint32 SetLocationAndColorData(uint8* Buffer, const TArray<FLidarPointCloudTrave
 		}
 		else
 		{
-			for (auto& Node : Nodes)
+			for (const FLidarPointCloudTraversalOctreeNode* Node : Nodes)
 			{
 				// Skip nodes with no available data
 				if (!Node->DataNode->HasData())
@@ -172,7 +172,7 @@ uint32 SetLocationAndColorData(uint8* Buffer, const TArray<FLidarPointCloudTrave
 /** Iterates over the provided nodes and sets scale data */
 void SetScaleData(uint8* Buffer, const TArray<FLidarPointCloudTraversalOctreeNode*>& Nodes)
 {
-	for (auto Node : Nodes)
+	for (const FLidarPointCloudTraversalOctreeNode* Node : Nodes)
 	{
 		// Skip nodes with no available data
 		if (!Node->DataNode->HasData())
@@ -363,7 +363,7 @@ void FLidarPointCloudTraversalOctree::GetVisibleNodes(TArray<FLidarPointCloudLOD
 
 		if (SelectionParams.MaxDepth < 0 || CurrentNode->Depth < SelectionParams.MaxDepth)
 		{
-			for (auto& Child : CurrentNode->Children)
+			for (FLidarPointCloudTraversalOctreeNode& Child : CurrentNode->Children)
 			{
 				Nodes.Enqueue(&Child);
 			}
@@ -427,10 +427,10 @@ void FLidarPointCloudLODManager::ProcessLOD(const TArray<FLidarPointCloudLODMana
 
 		for (int32 i = 0; i < InRegisteredProxies.Num(); ++i)
 		{
-			auto& RegisteredProxy = InRegisteredProxies[i];
+			const FLidarPointCloudLODManager::FRegisteredProxy& RegisteredProxy = InRegisteredProxies[i];
 
 			// Acquire a Shared Pointer from the Weak Pointer and check that it references a valid object
-			if (auto SceneProxyWrapper = RegisteredProxy.SceneProxyWrapper.Pin())
+			if (TSharedPtr<FLidarPointCloudSceneProxyWrapper, ESPMode::ThreadSafe> SceneProxyWrapper = RegisteredProxy.SceneProxyWrapper.Pin())
 			{
 				FScopeLock OctreeLock(&RegisteredProxy.PointCloud->Octree.DataLock);
 
@@ -484,7 +484,7 @@ void FLidarPointCloudLODManager::ProcessLOD(const TArray<FLidarPointCloudLODMana
 	}
 
 	// Used to pass render data updates to render thread
-	auto Buffer = BufferManager.GetFreeBuffer();
+	FLidarPointCloudDataBuffer* Buffer = BufferManager.GetFreeBuffer();
 	TArray<FLidarPointCloudProxyUpdateData> ProxyUpdateData;
 
 	// Build buffer data
@@ -503,12 +503,12 @@ void FLidarPointCloudLODManager::ProcessLOD(const TArray<FLidarPointCloudLODMana
 
 		for (int32 i = 0; i < SelectedNodesData.Num(); ++i)
 		{
-			auto& RegisteredProxy = InRegisteredProxies[i];
+			const FLidarPointCloudLODManager::FRegisteredProxy& RegisteredProxy = InRegisteredProxies[i];
 
 			// Only calculate if needed
 			if (RegisteredProxy.Component->PointSize > 0)
 			{
-				for (auto Node : SelectedNodesData[i])
+				for (FLidarPointCloudTraversalOctreeNode* Node : SelectedNodesData[i])
 				{
 					Node->CalculateVirtualDepth(RegisteredProxy.TraversalOctree->LevelWeights, RegisteredProxy.TraversalOctree->VirtualDepthMultiplier, RegisteredProxy.Component->PointSizeBias);
 				}
@@ -527,7 +527,7 @@ void FLidarPointCloudLODManager::ProcessLOD(const TArray<FLidarPointCloudLODMana
 				}
 
 				// Queue nodes to be streamed
-				for (auto Node : SelectedNodesData[i])
+				for (FLidarPointCloudTraversalOctreeNode* Node : SelectedNodesData[i])
 				{
 					RegisteredProxy.PointCloud->Octree.QueueNode(Node->DataNode, BulkDataLifetime);
 				}
@@ -550,7 +550,7 @@ void FLidarPointCloudLODManager::ProcessLOD(const TArray<FLidarPointCloudLODMana
 			{
 				UpdateData.Bounds.Reset(SelectedNodesData[i].Num());
 
-				for (auto Node : SelectedNodesData[i])
+				for (FLidarPointCloudTraversalOctreeNode* Node : SelectedNodesData[i])
 				{
 					FVector Extent = RegisteredProxy.TraversalOctree->Extents[Node->Depth];
 					UpdateData.Bounds.Emplace(Node->Center - Extent, Node->Center + Extent);
@@ -570,7 +570,7 @@ void FLidarPointCloudLODManager::ProcessLOD(const TArray<FLidarPointCloudLODMana
 	// Begin streaming data
 	for (int32 i = 0; i < InRegisteredProxies.Num(); ++i)
 	{
-		auto& RegisteredProxy = InRegisteredProxies[i];
+		const FLidarPointCloudLODManager::FRegisteredProxy& RegisteredProxy = InRegisteredProxies[i];
 
 		FScopeLock OctreeLock(&RegisteredProxy.PointCloud->Octree.DataLock);
 		RegisteredProxy.PointCloud->Octree.UnloadOldNodes(CurrentTime);
@@ -634,7 +634,7 @@ void FLidarPointCloudLODManager::PrepareProxies()
 		if (RegisteredProxy.Component->GetPointCloud())
 		{
 			// Acquire a Shared Pointer from the Weak Pointer and check that it references a valid object
-			if (auto SceneProxyWrapper = RegisteredProxy.SceneProxyWrapper.Pin())
+			if (TSharedPtr<FLidarPointCloudSceneProxyWrapper, ESPMode::ThreadSafe> SceneProxyWrapper = RegisteredProxy.SceneProxyWrapper.Pin())
 			{
 #if WITH_EDITOR
 				// Avoid doubling the point allocation of the same asset (once in Editor world and once in PIE world)
@@ -661,7 +661,7 @@ void FLidarPointCloudLODManager::PrepareProxies()
 				}
 
 				// If this is an editor component, use its own ViewportClient
-				if (auto Client = RegisteredProxy.Component->GetOwningViewportClient().Pin())
+				if (TSharedPtr<FViewportClient> Client = RegisteredProxy.Component->GetOwningViewportClient().Pin())
 				{
 					// If the ViewData cannot be successfully retrieved from the editor viewport, fall back to using main view
 					if (!RegisteredProxy.ViewData.ComputeFromEditorViewportClient(Client.Get()))
diff --git a/Engine/Plugins/Enterprise/LidarPointCloud/Source/LidarPointCloudRuntime/Private/LidarPointCloudOctree.cpp b/Engine/Plugins/Enterprise/LidarPointCloud/Source/LidarPointCloudRuntime/Private/LidarPointCloudOctree.cpp
--- a/Engine/Plugins/Enterprise/LidarPointCloud/Source/LidarPointCloudRuntime/Private/LidarPointCloudOctree.cpp
+++ b/Engine/Plugins/Enterprise/LidarPointCloud/Source/LidarPointCloudRuntime/Private/LidarPointCloudOctree.cpp
@@ -156,7 +156,7 @@ FLidarPointCloudOctreeNode* FLidarPointCloudOctreeNode::GetChildNodeAtLocation(c
 
 void FLidarPointCloudOctreeNode::InsertPoints(FLidarPointCloudOctree* Tree, const FLidarPointCloudPoint* Points, const int32& Count, ELidarPointCloudDuplicateHandling DuplicateHandling, const FVector& Translation)
 {
-	const auto& LODData = Tree->SharedData[Depth];
+	const FLidarPointCloudOctree::FSharedLODData& LODData = Tree->SharedData[Depth];
 
 	// Local 
 	TArray<FLidarPointCloudPoint> PointBuckets[8];
@@ -170,7 +170,7 @@ void FLidarPointCloudOctreeNode::InsertPoints(FLidarPointCloudOctree* Tree, cons
 	for (int32 Index = 0; Index < Count; ++Index)
 	{
 		const FVector AdjustedLocation = Points[Index].Location + Translation;
-		auto InGridData = CalculateGridCellData(AdjustedLocation, Center, LODData);
+		FGridAllocation InGridData = CalculateGridCellData(AdjustedLocation, Center, LODData);
 		FGridAllocation* GridCell = NewGridAllocationMap.Find(InGridData.Index);
 
 		// Attempt to allocate the point to this node
@@ -225,7 +225,7 @@ void FLidarPointCloudOctreeNode::InsertPoints(FLidarPointCloudOctree* Tree, cons
 		// Rebuild Current Grid Mapping
 		for (int32 i = 0; i < AllocatedPoints.Num(); ++i)
 		{
-			auto InGridData = CalculateGridCellData(AllocatedPoints[i].Location, Center, LODData);
+			FGridAllocation InGridData = CalculateGridCellData(AllocatedPoints[i].Location, Center, LODData);
 			FGridAllocation* GridCell = CurrentGridAllocationMap.Find(InGridData.Index);
 
 			// Attempt to allocate the point to this node
@@ -252,7 +252,7 @@ void FLidarPointCloudOctreeNode::InsertPoints(FLidarPointCloudOctree* Tree, cons
 		}
 
 		// Compare the incoming data to the currently held set, and replace if necessary
-		for (auto& Element : NewGridAllocationMap)
+		for (TPair<int32, FGridAllocation>& Element : NewGridAllocationMap)
 		{
 			const int32& GridIndex = Element.Key;
 			const FLidarPointCloudPoint& Point = Points[Element.Value.Index];
@@ -347,7 +347,7 @@ void FLidarPointCloudOctreeNode::InsertPoints(FLidarPointCloudOctree* Tree, cons
 
 void FLidarPointCloudOctreeNode::InsertPoints(FLidarPointCloudOctree* Tree, FLidarPointCloudPoint** Points, const int32& Count, ELidarPointCloudDuplicateHandling DuplicateHandling, const FVector& Translation)
 {
-	const auto& LODData = Tree->SharedData[Depth];
+	const FLidarPointCloudOctree::FSharedLODData& LODData = Tree->SharedData[Depth];
 
 	// Local 
 	TArray<FLidarPointCloudPoint> PointBuckets[8];
@@ -361,7 +361,7 @@ void FLidarPointCloudOctreeNode::InsertPoints(FLidarPointCloudOctree* Tree, FLid
 	for (int32 Index = 0; Index < Count; Index++)
 	{
 		const FVector AdjustedLocation = Points[Index]->Location + Translation;
-		auto InGridData = CalculateGridCellData(AdjustedLocation, Center, LODData);
+		FGridAllocation InGridData = CalculateGridCellData(AdjustedLocation, Center, LODData);
 		FGridAllocation* GridCell = NewGridAllocationMap.Find(InGridData.Index);
 
 		// Attempt to allocate the point to this node
@@ -416,7 +416,7 @@ void FLidarPointCloudOctreeNode::InsertPoints(FLidarPointCloudOctree* Tree, FLid
 		// Rebuild Current Grid Mapping
 		for (int32 i = 0; i < AllocatedPoints.Num(); ++i)
 		{
-			auto InGridData = CalculateGridCellData(AllocatedPoints[i].Location, Center, LODData);
+			FGridAllocation InGridData = CalculateGridCellData(AllocatedPoints[i].Location, Center, LODData);
 			FGridAllocation* GridCell = CurrentGridAllocationMap.Find(InGridData.Index);
 
 			// Attempt to allocate the point to this node
@@ -443,7 +443,7 @@ void FLidarPointCloudOctreeNode::InsertPoints(FLidarPointCloudOctree* Tree, FLid
 		}
 
 		// Compare the incoming data to the currently held set, and replace if necessary
-		for (auto& Element : NewGridAllocationMap)
+		for (TPair<int32, FGridAllocation>& Element : NewGridAllocationMap)
 		{
 			const int32& GridIndex = Element.Key;
 			const FLidarPointCloudPoint& Point = *Points[Element.Value.Index];
@@ -901,7 +901,7 @@ void FLidarPointCloudOctree::GetPointsAsCopies(TArray<FLidarPointCloudPoint>& Po
 				const int32 NumPointsToCopy = FMath::Max(0LL, FMath::Min((int64)CurrentNode->GetNumPoints(), Count + StartIndex) - StartIndex);
 				if (NumPointsToCopy > 0)
 				{
-					for (auto Point = CurrentNode->GetData() + StartIndex, DataEnd = Point + StartIndex + NumPointsToCopy; Point != DataEnd; ++Point)
+					for (FLidarPointCloudPoint* Point = CurrentNode->GetData() + StartIndex, *DataEnd = Point + StartIndex + NumPointsToCopy; Point != DataEnd; ++Point)
 					{
 						Points.Add(Point->Transform(*LocalToWorld));
 					}
@@ -1153,7 +1153,7 @@ void FLidarPointCloudOctree::SetVisibilityOfPointsByRay(const bool& bNewVisibili
 			// Sort points to speed up rendering
 			CurrentNode->SortVisiblePoints();
 
-			for (auto Child : CurrentNode->Children)
+			for (FLidarPointCloudOctreeNode*& Child : CurrentNode->Children)
 			{
 				Nodes.Enqueue(Child);
 			}
@@ -1342,7 +1342,7 @@ void FLidarPointCloudOctree::RemovePoints(TArray<FLidarPointCloudPoint*>& Points
 		return;
 	}
 
-	for (auto Point = Points.GetData(), DataEnd = Point + Points.Num(); Point != DataEnd; ++Point)
+	for (FLidarPointCloudPoint** Point = Points.GetData(), ** DataEnd = Point + Points.Num(); Point != DataEnd; ++Point)
 	{
 		(*Point)->bMarkedForDeletion = true;
 	}
@@ -1493,7 +1493,7 @@ void FLidarPointCloudOctree::Empty(bool bDestroyNodes)
 		Root.~FLidarPointCloudOctreeNode();
 
 		// Reset node counters
-		for (auto& Count : NodeCount)
+		for (FThreadSafeCounter& Count : NodeCount)
 		{
 			Count.Reset();
 		}
@@ -1511,7 +1511,7 @@ void FLidarPointCloudOctree::Empty(bool bDestroyNodes)
 	}
 	
 	// Reset point counters
-	for (auto& Count : PointCount)
+	for (FThreadSafeCounter64& Count : PointCount)
 	{
 		Count.Reset();
 	}
@@ -1525,7 +1525,7 @@ void FLidarPointCloudOctree::UnregisterTraversalOctree(FLidarPointCloudTraversal
 
 		for (int32 i = 0; i < LinkedTraversalOctrees.Num(); ++i)
 		{
-			if (auto TO = LinkedTraversalOctrees[i].Pin())
+			if (TSharedPtr<FLidarPointCloudTraversalOctree, ESPMode::ThreadSafe> TO = LinkedTraversalOctrees[i].Pin())
 			{
 				if (TO.Get() == TraversalOctree)
 				{
@@ -1609,7 +1609,7 @@ void FLidarPointCloudOctree::UnloadOldNodes(const float& CurrentTime)
 
 	for(int32 i = 0; i < NodesInUse.Num(); ++i)
 	{
-		auto& Node = NodesInUse[i];
+		FLidarPointCloudOctreeNode* Node = NodesInUse[i];
 
 		// Unload data, if it expired
 		if (Node->BulkDataLifetime < CurrentTime)
@@ -1677,7 +1677,7 @@ void FLidarPointCloudOctree::MarkTraversalOctreesForInvalidation()
 {
 	for (int32 i = 0; i < LinkedTraversalOctrees.Num(); ++i)
 	{
-		if (auto TraversalOctree = LinkedTraversalOctrees[i].Pin())
+		if (TSharedPtr<FLidarPointCloudTraversalOctree, ESPMode::ThreadSafe> TraversalOctree = LinkedTraversalOctrees[i].Pin())
 		{
 			TraversalOctree->bValid = false;
 		}
@@ -1868,7 +1868,7 @@ void FLidarPointCloudTraversalOctreeNode::CalculateVirtualDepth(const TArray<flo
 	Nodes.Enqueue(this);
 	while (Nodes.Dequeue(CurrentNode))
 	{
-		for (auto& Child : CurrentNode->Children)
+		for (const FLidarPointCloudTraversalOctreeNode& Child : CurrentNode->Children)
 		{
 			if (Child.bSelected)
 			{
@@ -1891,7 +1891,7 @@ void FLidarPointCloudTraversalOctreeNode::CalculateVirtualDepth(const TArray<flo
 	Nodes.Enqueue(this);
 	while (Nodes.Dequeue(CurrentNode))
 	{
-		for (auto& Child : CurrentNode->Children)
+		for (const FLidarPointCloudTraversalOctreeNode& Child : CurrentNode->Children)
 		{
 			if (Child.bSelected)
 			{
@@ -1935,7 +1935,7 @@ FLidarPointCloudTraversalOctree::FLidarPointCloudTraversalOctree(FLidarPointClou
 
 	int64 NumPoints = 0;
 	TArray<int64> PointCount;
-	for (auto& Count : Octree->PointCount)
+	for (FThreadSafeCounter64& Count : Octree->PointCount)
 	{
 		PointCount.Add(Count.GetValue());
 		NumPoints += Count.GetValue();
diff --git a/Engine/Plugins/Enterprise/LidarPointCloud/Source/LidarPointCloudRuntime/Private/LidarPointCloudOctreeMacros.h b/Engine/Plugins/Enterprise/LidarPointCloud/Source/LidarPointCloudRuntime/Private/LidarPointCloudOctreeMacros.h
--- a/Engine/Plugins/Enterprise/LidarPointCloud/Source/LidarPointCloudRuntime/Private/LidarPointCloudOctreeMacros.h
+++ b/Engine/Plugins/Enterprise/LidarPointCloud/Source/LidarPointCloudRuntime/Private/LidarPointCloudOctreeMacros.h
@@ -2,8 +2,8 @@
 
 #pragma once
 
-#define FOR(ElemName, Node) for(auto ElemName = Node->GetPersistentData(), DataEnd = ElemName + Node->GetNumPoints(); ElemName != DataEnd; ++ElemName)
-#define FOR_RO(ElemName, Node) for(auto ElemName = Node->GetData(), DataEnd = ElemName + Node->GetNumPoints(); ElemName != DataEnd; ++ElemName)
+#define FOR(ElemName, Node) for(FLidarPointCloudPoint* ElemName = Node->GetPersistentData(), * DataEnd = ElemName + Node->GetNumPoints(); ElemName != DataEnd; ++ElemName)
+#define FOR_RO(ElemName, Node) for(FLidarPointCloudPoint* ElemName = Node->GetData(), * DataEnd = ElemName + Node->GetNumPoints(); ElemName != DataEnd; ++ElemName)
 
 #define IS_VIS_CHECK_REQUIRED (bVisibleOnly && CurrentNode->NumVisiblePoints < CurrentNode->GetNumPoints())
 
@@ -18,7 +18,11 @@
 	while (Nodes.Dequeue(CurrentNode))\
 	{\
 		{ Action } \
-		for (auto Child : CurrentNode->Children) if (NodeTest) Nodes.Enqueue(Child);\
+		for (FLidarPointCloudOctreeNode* Child : CurrentNode->Children)\
+		{\
+			if (NodeTest)\
+			Nodes.Enqueue(Child);\
+		}\
 	}\
 }
 #define ITERATE_NODES(Action, NodeTest) ITERATE_NODES_BODY(Action, NodeTest,  )
@@ -90,7 +94,10 @@
 			if (!IS_VIS_CHECK_REQUIRED) { FOR##Mode(Point, CurrentNode) { if (POINT_BY_RAY) { Action } } }\
 			else { FOR##Mode(Point, CurrentNode) { if (Point->bVisible && POINT_BY_RAY) { Action } } }\
 \
-			for (auto Child : CurrentNode->Children) Nodes.Enqueue(Child);\
+			for (FLidarPointCloudOctreeNode* Child : CurrentNode->Children)\
+			{\
+				 Nodes.Enqueue(Child);\
+			}\
 		}\
 	}\
 }
diff --git a/Engine/Plugins/Enterprise/LidarPointCloud/Source/LidarPointCloudRuntime/Private/LidarPointCloudShared.cpp b/Engine/Plugins/Enterprise/LidarPointCloud/Source/LidarPointCloudRuntime/Private/LidarPointCloudShared.cpp
--- a/Engine/Plugins/Enterprise/LidarPointCloud/Source/LidarPointCloudRuntime/Private/LidarPointCloudShared.cpp
+++ b/Engine/Plugins/Enterprise/LidarPointCloud/Source/LidarPointCloudRuntime/Private/LidarPointCloudShared.cpp
@@ -83,12 +83,12 @@ FLidarPointCloudDataBufferManager::FLidarPointCloudDataBufferManager(const int32
 
 FLidarPointCloudDataBufferManager::~FLidarPointCloudDataBufferManager()
 {
-	auto Iterator = &Head;
+	TList<FLidarPointCloudDataBuffer>* Iterator = &Head;
 	while (Iterator)
 	{
 		if (Iterator != &Head)
 		{
-			auto Tmp = Iterator;
+			TList<FLidarPointCloudDataBuffer>* Tmp = Iterator;
 			Iterator = Iterator->Next;
 			delete Tmp;
 		}
@@ -105,7 +105,7 @@ FLidarPointCloudDataBuffer* FLidarPointCloudDataBufferManager::GetFreeBuffer()
 
 	// Find available memory allocation
 	{
-		auto Iterator = &Head;
+		TList<FLidarPointCloudDataBuffer>* Iterator = &Head;
 		while (Iterator)
 		{
 			if (!Iterator->Element.bInUse)
@@ -142,7 +142,7 @@ void FLidarPointCloudDataBufferManager::Resize(const int32& NewBufferSize)
 
 	BufferSize = NewBufferSize;
 
-	auto Iterator = &Head;
+	TList<FLidarPointCloudDataBuffer>* Iterator = &Head;
 	while (Iterator)
 	{
 		Iterator->Element.Resize(NewBufferSize);
diff --git a/Engine/Plugins/Enterprise/LidarPointCloud/Source/LidarPointCloudRuntime/Private/Rendering/LidarPointCloudRendering.cpp b/Engine/Plugins/Enterprise/LidarPointCloud/Source/LidarPointCloudRuntime/Private/Rendering/LidarPointCloudRendering.cpp
--- a/Engine/Plugins/Enterprise/LidarPointCloud/Source/LidarPointCloudRuntime/Private/Rendering/LidarPointCloudRendering.cpp
+++ b/Engine/Plugins/Enterprise/LidarPointCloud/Source/LidarPointCloudRuntime/Private/Rendering/LidarPointCloudRendering.cpp
@@ -45,7 +45,7 @@ public:
 		// Create collision visualization buffers
 		if (Octree->HasCollisionData())
 		{
-			auto CollisionData = Octree->GetCollisionData();
+			const FTriMeshCollisionData* CollisionData = Octree->GetCollisionData();
 			
 			// Release the existing resources, if present
 			VertexBuffer.ReleaseResource();
@@ -218,7 +218,7 @@ public:
 				// Draw selected nodes' bounds
 				if (Component->bDrawNodeBounds)
 				{
-					for (auto& Node : RenderData.Bounds)
+					for (const FBox& Node : RenderData.Bounds)
 					{
 						DrawWireBox(PDI, Node, FColor(72, 72, 255), SDPG_World);
 					}
@@ -304,7 +304,7 @@ public:
 
 		const bool bUsesSprites = Component->PointSize > 0;
 
-		auto BoundsSize = Component->GetPointCloud()->GetBounds().GetSize();
+		FVector BoundsSize = Component->GetPointCloud()->GetBounds().GetSize();
 
 		// Make sure to apply minimum bounds size
 		BoundsSize.X = FMath::Max(BoundsSize.X, 0.001f);
diff --git a/Engine/Plugins/Enterprise/LidarPointCloud/Source/LidarPointCloudRuntime/Public/IO/LidarPointCloudFileIO.h b/Engine/Plugins/Enterprise/LidarPointCloud/Source/LidarPointCloudRuntime/Public/IO/LidarPointCloudFileIO.h
--- a/Engine/Plugins/Enterprise/LidarPointCloud/Source/LidarPointCloudRuntime/Public/IO/LidarPointCloudFileIO.h
+++ b/Engine/Plugins/Enterprise/LidarPointCloud/Source/LidarPointCloudRuntime/Public/IO/LidarPointCloudFileIO.h
@@ -137,7 +137,7 @@ public:
 		FVector CenterOffset = Bounds.GetCenter();
 
 		// Apply it to the points
-		for (auto& Point : Points)
+		for (FLidarPointCloudPoint& Point : Points)
 		{
 			Point.Location -= CenterOffset;
 		}
