Lidar:
Fixed cloud being incorrectly culled when rotated

diff --git a/Engine/Plugins/Enterprise/LidarPointCloud/Source/LidarPointCloudRuntime/Private/LidarPointCloudLODManager.cpp b/Engine/Plugins/Enterprise/LidarPointCloud/Source/LidarPointCloudRuntime/Private/LidarPointCloudLODManager.cpp
--- a/Engine/Plugins/Enterprise/LidarPointCloud/Source/LidarPointCloudRuntime/Private/LidarPointCloudLODManager.cpp
+++ b/Engine/Plugins/Enterprise/LidarPointCloud/Source/LidarPointCloudRuntime/Private/LidarPointCloudLODManager.cpp
@@ -286,7 +286,7 @@ bool FLidarPointCloudViewData::ComputeFromEditorViewportClient(FViewportClient*
 void FLidarPointCloudTraversalOctree::GetVisibleNodes(TArray<FLidarPointCloudLODManager::FNodeSizeData>& NodeSizeData, const FLidarPointCloudViewData& ViewData, const int32& ProxyIndex, const FLidarPointCloudNodeSelectionParams& SelectionParams, const float& CurrentTime)
 {
 	// Skip processing, if the asset is not visible
-	if (!ViewData.ViewFrustum.IntersectBox(GetCenter(), Octree->GetExtent()))
+	if (!ViewData.ViewFrustum.IntersectBox(GetCenter(), GetExtent()))
 	{
 		return;
 	}
diff --git a/Engine/Plugins/Enterprise/LidarPointCloud/Source/LidarPointCloudRuntime/Public/LidarPointCloudOctree.h b/Engine/Plugins/Enterprise/LidarPointCloud/Source/LidarPointCloudRuntime/Public/LidarPointCloudOctree.h
--- a/Engine/Plugins/Enterprise/LidarPointCloud/Source/LidarPointCloudRuntime/Public/LidarPointCloudOctree.h
+++ b/Engine/Plugins/Enterprise/LidarPointCloud/Source/LidarPointCloudRuntime/Public/LidarPointCloudOctree.h
@@ -516,4 +516,5 @@ struct FLidarPointCloudTraversalOctree
 	void GetVisibleNodes(TArray<FLidarPointCloudLODManager::FNodeSizeData>& NodeSizeData, const FLidarPointCloudViewData& ViewData, const int32& ProxyIndex, const FLidarPointCloudNodeSelectionParams& SelectionParams, const float& CurrentTime);
 
 	FVector GetCenter() const { return Root.Center; }
+	FVector GetExtent() const { return Extents[0]; }
 };
