Lidar:
Fixed few API related issues

diff --git a/Engine/Plugins/Enterprise/LidarPointCloud/Shaders/Private/LidarPointCloudVertexFactory.ush b/Engine/Plugins/Enterprise/LidarPointCloud/Shaders/Private/LidarPointCloudVertexFactory.ush
--- a/Engine/Plugins/Enterprise/LidarPointCloud/Shaders/Private/LidarPointCloudVertexFactory.ush
+++ b/Engine/Plugins/Enterprise/LidarPointCloud/Shaders/Private/LidarPointCloudVertexFactory.ush
@@ -257,6 +257,9 @@ FMaterialVertexParameters GetMaterialVertexParameters(FVertexFactoryInput Input,
 
 #if NUM_MATERIAL_TEXCOORDS_VERTEX
 	Result.TexCoords[0] = PrecomputedUV_Quad[Input.VertexId % 4];
+	#if NUM_MATERIAL_TEXCOORDS_VERTEX > 1
+		Result.TexCoords[1] = Intermediates.Scale;
+	#endif
 #endif
 
     return Result;
diff --git a/Engine/Plugins/Enterprise/LidarPointCloud/Source/LidarPointCloudRuntime/Private/IO/LidarPointCloudFileIO_ASCII.cpp b/Engine/Plugins/Enterprise/LidarPointCloud/Source/LidarPointCloudRuntime/Private/IO/LidarPointCloudFileIO_ASCII.cpp
--- a/Engine/Plugins/Enterprise/LidarPointCloud/Source/LidarPointCloudRuntime/Private/IO/LidarPointCloudFileIO_ASCII.cpp
+++ b/Engine/Plugins/Enterprise/LidarPointCloud/Source/LidarPointCloudRuntime/Private/IO/LidarPointCloudFileIO_ASCII.cpp
@@ -32,6 +32,14 @@
 		PC_STREAMTEXTFILE_OFFSET(OldBufferSize)														\
 	}
 
+TSharedPtr<FLidarPointCloudImportSettings_ASCII> MakeImportSettings(const FString& Filename, const FVector2D& RGBRange, const FLidarPointCloudImportSettings_ASCII_Columns& Columns)
+{
+	TSharedPtr<FLidarPointCloudImportSettings_ASCII> ImportSettings = MakeShared<FLidarPointCloudImportSettings_ASCII>(Filename);
+	ImportSettings->RGBRange = RGBRange;
+	ImportSettings->SelectedColumns = { FMath::Max(-1, Columns.LocationX), FMath::Max(-1, Columns.LocationY), FMath::Max(-1, Columns.LocationZ), FMath::Max(-1, Columns.Red), FMath::Max(-1, Columns.Green), FMath::Max(-1, Columns.Blue), FMath::Max(-1, Columns.Intensity) };
+	return ImportSettings;
+}
+
 bool ULidarPointCloudFileIO_ASCII::HandleImport(const FString& Filename, TSharedPtr<FLidarPointCloudImportSettings> ImportSettings, FLidarPointCloudImportResults &OutImportResults)
 {
 	if (!ValidateImportSettings(ImportSettings, Filename))
@@ -275,11 +283,12 @@ bool ULidarPointCloudFileIO_ASCII::HandleExport(const FString& Filename, class U
 
 void ULidarPointCloudFileIO_ASCII::CreatePointCloudFromFile(UObject* WorldContextObject, const FString& Filename, bool bUseAsync, FVector2D RGBRange, FLidarPointCloudImportSettings_ASCII_Columns Columns, FLatentActionInfo LatentInfo, ELidarPointCloudAsyncMode& AsyncMode, float& Progress, ULidarPointCloud*& PointCloud)
 {
-	TSharedPtr<FLidarPointCloudImportSettings_ASCII> ImportSettings = MakeShared<FLidarPointCloudImportSettings_ASCII>(Filename);
-	ImportSettings->RGBRange = RGBRange;
-	ImportSettings->SelectedColumns = { FMath::Max(-1, Columns.LocationX), FMath::Max(-1, Columns.LocationY), FMath::Max(-1, Columns.LocationZ), FMath::Max(-1, Columns.Red), FMath::Max(-1, Columns.Green), FMath::Max(-1, Columns.Blue), FMath::Max(-1, Columns.Intensity) };
+	ULidarPointCloudBlueprintLibrary::CreatePointCloudFromFile(WorldContextObject, Filename, bUseAsync, LatentInfo, MakeImportSettings(Filename, RGBRange, Columns), AsyncMode, Progress, PointCloud);
+}
 
-	ULidarPointCloudBlueprintLibrary::CreatePointCloudFromFile(WorldContextObject, Filename, bUseAsync, LatentInfo, ImportSettings, AsyncMode, Progress, PointCloud);
+ULidarPointCloud* ULidarPointCloudFileIO_ASCII::CreatePointCloudFromFile(const FString& Filename, const FLidarPointCloudAsyncParameters& AsyncParameters, const FVector2D& RGBRange, const FLidarPointCloudImportSettings_ASCII_Columns& Columns)
+{
+	return ULidarPointCloud::CreateFromFile(Filename, AsyncParameters, MakeImportSettings(Filename, RGBRange, Columns));
 }
 
 FLidarPointCloudImportSettings_ASCII::FLidarPointCloudImportSettings_ASCII(const FString& Filename)
diff --git a/Engine/Plugins/Enterprise/LidarPointCloud/Source/LidarPointCloudRuntime/Private/LidarPointCloudOctree.cpp b/Engine/Plugins/Enterprise/LidarPointCloud/Source/LidarPointCloudRuntime/Private/LidarPointCloudOctree.cpp
--- a/Engine/Plugins/Enterprise/LidarPointCloud/Source/LidarPointCloudRuntime/Private/LidarPointCloudOctree.cpp
+++ b/Engine/Plugins/Enterprise/LidarPointCloud/Source/LidarPointCloudRuntime/Private/LidarPointCloudOctree.cpp
@@ -1125,6 +1125,54 @@ void FLidarPointCloudOctree::SetVisibilityOfPointsInBox(const bool& bNewVisibili
 	}, NODE_IN_BOX);
 }
 
+void FLidarPointCloudOctree::SetVisibilityOfFirstPointByRay(const bool& bNewVisibility, const FLidarPointCloudRay& Ray, const float& Radius)
+{
+	bool bPointProcessed = false;
+	const float RadiusSq = Radius * Radius;
+	ITERATE_NODES({
+		// Skip node if it already has all points set to the required visibility state
+		bool bSkipNode = (CurrentNode->NumVisiblePoints == CurrentNode->GetNumPoints() && bNewVisibility) || (CurrentNode->NumVisiblePoints == 0 && !bNewVisibility);
+		if (!bSkipNode && Ray.Intersects(CurrentNode->GetBounds(this)))
+		{
+			CurrentNode->NumVisiblePoints = 0;
+
+			FOR(Point, CurrentNode)
+			{
+				if (Point->bVisible != bNewVisibility && POINT_BY_RAY)
+				{
+					Point->bVisible = bNewVisibility;
+					bPointProcessed = true;
+				}
+
+				if (Point->bVisible)
+				{
+					++CurrentNode->NumVisiblePoints;
+				}
+
+				if (bPointProcessed)
+				{
+					break;
+				}
+			}
+
+			CurrentNode->bVisibilityDirty = false;
+
+			// Sort points to speed up rendering
+			CurrentNode->SortVisiblePoints();
+
+			if (bPointProcessed)
+			{
+				return;
+			}
+
+			for (FLidarPointCloudOctreeNode*& Child : CurrentNode->Children)
+			{
+				Nodes.Enqueue(Child);
+			}
+		}
+	}, false);
+}
+
 void FLidarPointCloudOctree::SetVisibilityOfPointsByRay(const bool& bNewVisibility, const FLidarPointCloudRay& Ray, const float& Radius)
 {
 	const float RadiusSq = Radius * Radius;
@@ -1208,6 +1256,14 @@ void FLidarPointCloudOctree::ExecuteActionOnPointsInBox(TFunction<void(FLidarPoi
 	PROCESS_IN_BOX({ Action(Point); });
 }
 
+void FLidarPointCloudOctree::ExecuteActionOnFirstPointByRay(TFunction<void(FLidarPointCloudPoint*)> Action, const FLidarPointCloudRay& Ray, const float& Radius, bool bVisibleOnly)
+{
+	if (FLidarPointCloudPoint* Point = RaycastSingle(Ray, Radius, bVisibleOnly))
+	{
+		Action(Point);
+	}
+}
+
 void FLidarPointCloudOctree::ExecuteActionOnPointsByRay(TFunction<void(FLidarPointCloudPoint*)> Action, const FLidarPointCloudRay& Ray, const float& Radius, bool bVisibleOnly)
 {
 	PROCESS_BY_RAY({ Action(Point); });
@@ -1228,6 +1284,14 @@ void FLidarPointCloudOctree::ApplyColorToPointsInBox(const FColor& NewColor, con
 	PROCESS_IN_BOX({ Point->Color = NewColor; });
 }
 
+void FLidarPointCloudOctree::ApplyColorToFirstPointByRay(const FColor& NewColor, const FLidarPointCloudRay& Ray, const float& Radius, bool bVisibleOnly)
+{
+	if (FLidarPointCloudPoint* Point = RaycastSingle(Ray, Radius, bVisibleOnly))
+	{
+		Point->Color = NewColor;
+	}
+}
+
 void FLidarPointCloudOctree::ApplyColorToPointsByRay(const FColor& NewColor, const FLidarPointCloudRay& Ray, const float& Radius, bool bVisibleOnly)
 {
 	PROCESS_BY_RAY({ Point->Color = NewColor; });
diff --git a/Engine/Plugins/Enterprise/LidarPointCloud/Source/LidarPointCloudRuntime/Public/IO/LidarPointCloudFileIO_ASCII.h b/Engine/Plugins/Enterprise/LidarPointCloud/Source/LidarPointCloudRuntime/Public/IO/LidarPointCloudFileIO_ASCII.h
--- a/Engine/Plugins/Enterprise/LidarPointCloud/Source/LidarPointCloudRuntime/Public/IO/LidarPointCloudFileIO_ASCII.h
+++ b/Engine/Plugins/Enterprise/LidarPointCloud/Source/LidarPointCloudRuntime/Public/IO/LidarPointCloudFileIO_ASCII.h
@@ -4,6 +4,7 @@
 
 #include "CoreMinimal.h"
 #include "LidarPointCloudShared.h"
+#include "LidarPointCloud.h"
 #include "LidarPointCloudFileIO.h"
 #include "Kismet/BlueprintFunctionLibrary.h"
 #include "Engine/LatentActionManager.h"
@@ -157,6 +158,7 @@ public:
 
 	UFUNCTION(BlueprintCallable, Category = "Lidar Point Cloud", meta = (Latent, WorldContext = "WorldContextObject", LatentInfo = "LatentInfo", ExpandEnumAsExecs = "AsyncMode", DisplayName = "Create Lidar Point Cloud From File (ASCII)"))
 	static void CreatePointCloudFromFile(UObject* WorldContextObject, const FString& Filename, bool bUseAsync, FVector2D RGBRange, FLidarPointCloudImportSettings_ASCII_Columns Columns, FLatentActionInfo LatentInfo, ELidarPointCloudAsyncMode& AsyncMode, float& Progress, ULidarPointCloud*& PointCloud);
+	static ULidarPointCloud* CreatePointCloudFromFile(const FString& Filename, const FLidarPointCloudAsyncParameters& AsyncParameters, const FVector2D& RGBRange, const FLidarPointCloudImportSettings_ASCII_Columns& Columns);
 
 	virtual bool HandleImport(const FString& Filename, TSharedPtr<FLidarPointCloudImportSettings> ImportSettings, FLidarPointCloudImportResults &OutImportResults) override;
 	virtual TSharedPtr<FLidarPointCloudImportSettings> GetImportSettings(const FString& Filename) override { return TSharedPtr<FLidarPointCloudImportSettings>(new FLidarPointCloudImportSettings_ASCII(Filename)); }
diff --git a/Engine/Plugins/Enterprise/LidarPointCloud/Source/LidarPointCloudRuntime/Public/LidarPointCloud.h b/Engine/Plugins/Enterprise/LidarPointCloud/Source/LidarPointCloudRuntime/Public/LidarPointCloud.h
--- a/Engine/Plugins/Enterprise/LidarPointCloud/Source/LidarPointCloudRuntime/Public/LidarPointCloud.h
+++ b/Engine/Plugins/Enterprise/LidarPointCloud/Source/LidarPointCloudRuntime/Public/LidarPointCloud.h
@@ -293,6 +293,11 @@ public:
 	void SetVisibilityOfPointsInBox(bool bNewVisibility, FVector Center, FVector Extent) { SetVisibilityOfPointsInBox(bNewVisibility, FBox(Center - Extent, Center + Extent)); }
 	void SetVisibilityOfPointsInBox(const bool& bNewVisibility, const FBox& Box) { Octree.SetVisibilityOfPointsInBox(bNewVisibility, Box.ShiftBy(-LocationOffset.ToVector())); }
 
+	/** Sets visibility of the first point hit by the given ray. */
+	UFUNCTION(BlueprintCallable, Category = "Lidar Point Cloud")
+	void SetVisibilityOfFirstPointByRay(bool bNewVisibility, FVector Origin, FVector Direction, float Radius) { SetVisibilityOfFirstPointByRay(bNewVisibility, FLidarPointCloudRay(Origin, Direction), Radius); }
+	void SetVisibilityOfFirstPointByRay(const bool& bNewVisibility, const FLidarPointCloudRay& Ray, const float& Radius) { Octree.SetVisibilityOfFirstPointByRay(bNewVisibility, Ray.ShiftBy(-LocationOffset.ToVector()), Radius); }
+
 	/** Sets visibility of points hit by the given ray. */
 	UFUNCTION(BlueprintCallable, Category = "Lidar Point Cloud")
 	void SetVisibilityOfPointsByRay(bool bNewVisibility, FVector Origin, FVector Direction, float Radius) { SetVisibilityOfPointsByRay(bNewVisibility, FLidarPointCloudRay(Origin, Direction), Radius); }
@@ -321,6 +326,9 @@ public:
 	void ExecuteActionOnPointsInBox(TFunction<void(FLidarPointCloudPoint*)> Action, const FVector& Center, const FVector& Extent, const bool& bVisibleOnly) { ExecuteActionOnPointsInBox(MoveTemp(Action), FBox(Center - Extent, Center + Extent), bVisibleOnly); }
 	void ExecuteActionOnPointsInBox(TFunction<void(FLidarPointCloudPoint*)> Action, const FBox& Box, const bool& bVisibleOnly) { Octree.ExecuteActionOnPointsInBox(MoveTemp(Action), Box.ShiftBy(-LocationOffset.ToVector()), bVisibleOnly); }
 
+	/** Executes the provided action on the first point hit by the given ray. */
+	void ExecuteActionOnFirstPointByRay(TFunction<void(FLidarPointCloudPoint*)> Action, const FLidarPointCloudRay& Ray, const float& Radius, bool bVisibleOnly) { Octree.ExecuteActionOnFirstPointByRay(MoveTemp(Action), Ray.ShiftBy(-LocationOffset.ToVector()), Radius, bVisibleOnly); }
+
 	/** Executes the provided action on each of the points hit by the given ray. */
 	void ExecuteActionOnPointsByRay(TFunction<void(FLidarPointCloudPoint*)> Action, const FLidarPointCloudRay& Ray, const float& Radius, bool bVisibleOnly) { Octree.ExecuteActionOnPointsByRay(MoveTemp(Action), Ray.ShiftBy(-LocationOffset.ToVector()), Radius, bVisibleOnly); }
 
@@ -342,6 +350,11 @@ public:
 	void ApplyColorToPointsInBox(FColor NewColor, FVector Center, FVector Extent, bool bVisibleOnly) { ApplyColorToPointsInBox(NewColor, FBox(Center - Extent, Center + Extent), bVisibleOnly); }
 	void ApplyColorToPointsInBox(const FColor& NewColor, const FBox& Box, const bool& bVisibleOnly) { Octree.ApplyColorToPointsInBox(NewColor, Box.ShiftBy(-LocationOffset.ToVector()), bVisibleOnly); }
 
+	/** Applies the given color to the first point hit by the given ray */
+	UFUNCTION(BlueprintCallable, Category = "Lidar Point Cloud")
+	void ApplyColorToFirstPointByRay(FColor NewColor, FVector Origin, FVector Direction, float Radius, bool bVisibleOnly) { ApplyColorToFirstPointByRay(NewColor, FLidarPointCloudRay(Origin, Direction), Radius, bVisibleOnly); }
+	void ApplyColorToFirstPointByRay(const FColor& NewColor, const FLidarPointCloudRay& Ray, const float& Radius, bool bVisibleOnly) { Octree.ApplyColorToFirstPointByRay(NewColor, Ray.ShiftBy(-LocationOffset.ToVector()), Radius, bVisibleOnly); }
+
 	/** Applies the given color to all points hit by the given ray */
 	UFUNCTION(BlueprintCallable, Category = "Lidar Point Cloud")
 	void ApplyColorToPointsByRay(FColor NewColor, FVector Origin, FVector Direction, float Radius, bool bVisibleOnly) { ApplyColorToPointsByRay(NewColor, FLidarPointCloudRay(Origin, Direction), Radius, bVisibleOnly); }
diff --git a/Engine/Plugins/Enterprise/LidarPointCloud/Source/LidarPointCloudRuntime/Public/LidarPointCloudComponent.h b/Engine/Plugins/Enterprise/LidarPointCloud/Source/LidarPointCloudRuntime/Public/LidarPointCloudComponent.h
--- a/Engine/Plugins/Enterprise/LidarPointCloud/Source/LidarPointCloudRuntime/Public/LidarPointCloudComponent.h
+++ b/Engine/Plugins/Enterprise/LidarPointCloud/Source/LidarPointCloudRuntime/Public/LidarPointCloudComponent.h
@@ -282,6 +282,17 @@ public:
 		}
 	}
 
+	/** Sets visibility of the first point hit by the given ray. */
+	UFUNCTION(BlueprintCallable, Category = "Lidar Point Cloud")
+	void SetVisibilityOfFirstPointByRay(bool bNewVisibility, FVector Origin, FVector Direction, float Radius) { SetVisibilityOfFirstPointByRay(bNewVisibility, FLidarPointCloudRay(Origin, Direction), Radius); }
+	void SetVisibilityOfFirstPointByRay(bool bNewVisibility, FLidarPointCloudRay Ray, float Radius)
+	{
+		if (PointCloud)
+		{
+			PointCloud->SetVisibilityOfFirstPointByRay(bNewVisibility, Ray.TransformBy(GetComponentTransform().Inverse()), Radius);
+		}
+	}
+
 	/** Sets visibility of points hit by the given ray. */
 	UFUNCTION(BlueprintCallable, Category = "Lidar Point Cloud")
 	void SetVisibilityOfPointsByRay(bool bNewVisibility, FVector Origin, FVector Direction, float Radius) { SetVisibilityOfPointsByRay(bNewVisibility, FLidarPointCloudRay(Origin, Direction), Radius); }
@@ -313,6 +324,15 @@ public:
 		}
 	}
 
+	/** Executes the provided action on the first point hit by the given ray. */
+	void ExecuteActionOnFirstPointByRay(TFunction<void(FLidarPointCloudPoint*)> Action, FLidarPointCloudRay Ray, float Radius, bool bVisibleOnly)
+	{
+		if (PointCloud)
+		{
+			PointCloud->ExecuteActionOnFirstPointByRay(MoveTemp(Action), Ray.TransformBy(GetComponentTransform().Inverse()), Radius, bVisibleOnly);
+		}
+	}
+
 	/** Executes the provided action on each of the points hit by the given ray. */
 	void ExecuteActionOnPointsByRay(TFunction<void(FLidarPointCloudPoint*)> Action, FLidarPointCloudRay Ray, float Radius, bool bVisibleOnly)
 	{
@@ -344,6 +364,17 @@ public:
 		}
 	}
 
+	/** Applies the given color to the first point hit by the given ray */
+	UFUNCTION(BlueprintCallable, Category = "Lidar Point Cloud")
+	void ApplyColorToFirstPointByRay(FColor NewColor, FVector Origin, FVector Direction, float Radius, bool bVisibleOnly) { ApplyColorToFirstPointByRay(NewColor, FLidarPointCloudRay(Origin, Direction), Radius, bVisibleOnly); }
+	void ApplyColorToFirstPointByRay(const FColor& NewColor, FLidarPointCloudRay Ray, float Radius, bool bVisibleOnly)
+	{
+		if (PointCloud)
+		{
+			PointCloud->ApplyColorToFirstPointByRay(NewColor, Ray.TransformBy(GetComponentTransform().Inverse()), Radius, bVisibleOnly);
+		}
+	}
+
 	/** Applies the given color to all points hit by the given ray */
 	UFUNCTION(BlueprintCallable, Category = "Lidar Point Cloud")
 	void ApplyColorToPointsByRay(FColor NewColor, FVector Origin, FVector Direction, float Radius, bool bVisibleOnly) { ApplyColorToPointsByRay(NewColor, FLidarPointCloudRay(Origin, Direction), Radius, bVisibleOnly); }
diff --git a/Engine/Plugins/Enterprise/LidarPointCloud/Source/LidarPointCloudRuntime/Public/LidarPointCloudOctree.h b/Engine/Plugins/Enterprise/LidarPointCloud/Source/LidarPointCloudRuntime/Public/LidarPointCloudOctree.h
--- a/Engine/Plugins/Enterprise/LidarPointCloud/Source/LidarPointCloudRuntime/Public/LidarPointCloudOctree.h
+++ b/Engine/Plugins/Enterprise/LidarPointCloud/Source/LidarPointCloudRuntime/Public/LidarPointCloudOctree.h
@@ -312,6 +312,9 @@ public:
 	/** Sets visibility of points within the given box. */
 	void SetVisibilityOfPointsInBox(const bool& bNewVisibility, const FBox& Box);
 
+	/** Sets visibility of the first point hit by the given ray. */
+	void SetVisibilityOfFirstPointByRay(const bool& bNewVisibility, const FLidarPointCloudRay& Ray, const float& Radius);
+
 	/** Sets visibility of points hit by the given ray. */
 	void SetVisibilityOfPointsByRay(const bool& bNewVisibility, const FLidarPointCloudRay& Ray, const float& Radius);
 
@@ -330,6 +333,9 @@ public:
 	/** Executes the provided action on each of the points within the given box. */
 	void ExecuteActionOnPointsInBox(TFunction<void(FLidarPointCloudPoint*)> Action, const FBox& Box, const bool& bVisibleOnly);
 
+	/** Executes the provided action on the first point hit by the given ray. */
+	void ExecuteActionOnFirstPointByRay(TFunction<void(FLidarPointCloudPoint*)> Action, const FLidarPointCloudRay& Ray, const float& Radius, bool bVisibleOnly);
+
 	/** Executes the provided action on each of the points hit by the given ray. */
 	void ExecuteActionOnPointsByRay(TFunction<void(FLidarPointCloudPoint*)> Action, const FLidarPointCloudRay& Ray, const float& Radius, bool bVisibleOnly);
 
@@ -342,6 +348,9 @@ public:
 	/** Applies the given color to all points within the box */
 	void ApplyColorToPointsInBox(const FColor& NewColor, const FBox& Box, const bool& bVisibleOnly);
 
+	/** Applies the given color to the first point hit by the given ray */
+	void ApplyColorToFirstPointByRay(const FColor& NewColor, const FLidarPointCloudRay& Ray, const float& Radius, bool bVisibleOnly);
+
 	/** Applies the given color to all points hit by the given ray */
 	void ApplyColorToPointsByRay(const FColor& NewColor, const FLidarPointCloudRay& Ray, const float& Radius, bool bVisibleOnly);
 
