Lidar:
Fixed content not cooking with LiDAR Point Cloud plugin enabled

diff --git a/Engine/Plugins/Enterprise/LidarPointCloud/Source/LidarPointCloudRuntime/Private/LidarPointCloudSettings.cpp b/Engine/Plugins/Enterprise/LidarPointCloud/Source/LidarPointCloudRuntime/Private/LidarPointCloudSettings.cpp
--- a/Engine/Plugins/Enterprise/LidarPointCloud/Source/LidarPointCloudRuntime/Private/LidarPointCloudSettings.cpp
+++ b/Engine/Plugins/Enterprise/LidarPointCloud/Source/LidarPointCloudRuntime/Private/LidarPointCloudSettings.cpp
@@ -23,6 +23,5 @@ ULidarPointCloudSettings::ULidarPointCloudSettings()
 	, MaxNumberOfPointsToScanASCII(100000)
 	, ExportScale(0.01f)
 	, bUseIOCaching(false)
-	, bLogBufferExpansion(false)
 {
 }
diff --git a/Engine/Plugins/Enterprise/LidarPointCloud/Source/LidarPointCloudRuntime/Private/Rendering/LidarPointCloudRenderBuffers.cpp b/Engine/Plugins/Enterprise/LidarPointCloud/Source/LidarPointCloudRuntime/Private/Rendering/LidarPointCloudRenderBuffers.cpp
--- a/Engine/Plugins/Enterprise/LidarPointCloud/Source/LidarPointCloudRuntime/Private/Rendering/LidarPointCloudRenderBuffers.cpp
+++ b/Engine/Plugins/Enterprise/LidarPointCloud/Source/LidarPointCloudRuntime/Private/Rendering/LidarPointCloudRenderBuffers.cpp
@@ -23,40 +23,21 @@ TGlobalResource<FLidarPointCloudRenderBuffer> GLidarPointCloudRenderBuffer;
 TGlobalResource<FLidarPointCloudIndexBuffer> GLidarPointCloudIndexBuffer;
 TGlobalResource<FLidarPointCloudVertexFactory> GLidarPointCloudVertexFactory;
 
-void FLidarPointCloudBuffer::Resize(const uint32& RequestedCapacity)
+//////////////////////////////////////////////////////////// Index Buffer
+
+void FLidarPointCloudIndexBuffer::Resize(const uint32 & RequestedCapacity)
 {
 	// This must be called from Rendering thread
 	check(IsInRenderingThread());
 
 	if (Capacity != RequestedCapacity)
 	{
-		if (GetDefault<ULidarPointCloudSettings>()->bLogBufferExpansion)
-		{
-			PC_LOG("Resizing %s: %u => %u", *GetFriendlyName(), Capacity, RequestedCapacity);
-		}
-
-		Release();
+		ReleaseResource();
 		Capacity = RequestedCapacity;
-		Initialize();
+		InitResource();
 	}
 }
 
-FLidarPointCloudBuffer::~FLidarPointCloudBuffer()
-{
-	FRenderCommandFence Fence;
-
-	ENQUEUE_RENDER_COMMAND(ReleaseLidarPointCloudRenderBuffer)(
-		[this](FRHICommandListImmediate& RHICmdList)
-		{
-			Release();
-		});
-
-	Fence.BeginFence();
-	Fence.Wait();
-}
-
-//////////////////////////////////////////////////////////// Index Buffer
-
 void FLidarPointCloudIndexBuffer::InitRHI()
 {
 	FRHIResourceCreateInfo CreateInfo;
@@ -92,14 +73,27 @@ void FLidarPointCloudIndexBuffer::InitRHI()
 
 //////////////////////////////////////////////////////////// Structured Buffer
 
+void FLidarPointCloudRenderBuffer::Resize(const uint32& RequestedCapacity)
+{
+	// This must be called from Rendering thread
+	check(IsInRenderingThread());
+
+	if (Capacity != RequestedCapacity)
+	{
+		ReleaseResource();
+		Capacity = RequestedCapacity;
+		InitResource();
+	}
+}
+
 void FLidarPointCloudRenderBuffer::InitRHI()
 {
 	// This must be called from Rendering thread
 	check(IsInRenderingThread());
 
 	FRHIResourceCreateInfo CreateInfo;
-	Buffer = RHICreateVertexBuffer(ElementSize * Capacity, BUF_ShaderResource | BUF_Dynamic, CreateInfo);
-	SRV = RHICreateShaderResourceView(Buffer, ElementSize, PF_R32_FLOAT);
+	Buffer = RHICreateVertexBuffer(sizeof(uint32) * Capacity, BUF_ShaderResource | BUF_Dynamic, CreateInfo);
+	SRV = RHICreateShaderResourceView(Buffer, sizeof(uint32), PF_R32_FLOAT);
 }
 
 void FLidarPointCloudRenderBuffer::ReleaseRHI()
diff --git a/Engine/Plugins/Enterprise/LidarPointCloud/Source/LidarPointCloudRuntime/Private/Rendering/LidarPointCloudRenderBuffers.h b/Engine/Plugins/Enterprise/LidarPointCloud/Source/LidarPointCloudRuntime/Private/Rendering/LidarPointCloudRenderBuffers.h
--- a/Engine/Plugins/Enterprise/LidarPointCloud/Source/LidarPointCloudRuntime/Private/Rendering/LidarPointCloudRenderBuffers.h
+++ b/Engine/Plugins/Enterprise/LidarPointCloud/Source/LidarPointCloudRuntime/Private/Rendering/LidarPointCloudRenderBuffers.h
@@ -15,25 +15,13 @@
 class FLidarPointCloudBuffer
 {
 public:
-	/** Resizes the buffer to the specified capacity, if necessary. Must be called from Rendering thread. */
-	void Resize(const uint32& RequestedCapacity);
-
-	FLidarPointCloudBuffer(const uint32& ElementSize) : Capacity(100000), ElementSize(ElementSize) {}
-
-	virtual FString GetFriendlyName() const { return TEXT("FPointCloudBuffer"); }
+	FLidarPointCloudBuffer() : Capacity(100000) {}
 
-	/** Returns the total size of this buffer */
-	virtual uint32 GetSize() const { return Capacity * ElementSize; }
+	/** Resizes the buffer to the specified capacity, if necessary. Must be called from Rendering thread. */
+	virtual void Resize(const uint32& RequestedCapacity) = 0;
 
 protected:
 	uint32 Capacity;
-	uint32 ElementSize;
-
-	FLidarPointCloudBuffer() : FLidarPointCloudBuffer(0) {}
-	virtual ~FLidarPointCloudBuffer();
-
-	virtual void Initialize() = 0;
-	virtual void Release() = 0;
 };
 
 /**
@@ -42,18 +30,10 @@ protected:
 class FLidarPointCloudIndexBuffer : public FLidarPointCloudBuffer, public FIndexBuffer
 {
 public:
-	FLidarPointCloudIndexBuffer() : FLidarPointCloudBuffer(sizeof(uint32)) { }
-
-	uint32 GetPointOffset() const { return PointOffset; }
-
-private:
-	virtual void InitRHI() override;
-	virtual void Initialize() override { InitResource(); }
-	virtual void Release() override { ReleaseResource(); }
-	virtual FString GetFriendlyName() const override { return TEXT("FLidarPointCloudIndexBuffer"); }
-
-private:
 	uint32 PointOffset;
+
+	virtual void Resize(const uint32& RequestedCapacity) override;
+	virtual void InitRHI() override;	
 };
 
 /**
@@ -67,21 +47,9 @@ public:
 
 	int32 PointCount;
 
-	FLidarPointCloudRenderBuffer() : FLidarPointCloudBuffer(4) {}
-	virtual ~FLidarPointCloudRenderBuffer() { Release(); }
-
+	virtual void Resize(const uint32& RequestedCapacity) override;	
 	virtual void InitRHI() override;
 	virtual void ReleaseRHI() override;
-
-	virtual void Initialize() override { InitResource(); }
-	virtual void Release() override { ReleaseResource(); }
-	virtual FString GetFriendlyName() const override { return TEXT("FLidarPointCloudRenderBuffer"); }
-
-private:
-	FLidarPointCloudRenderBuffer(const FLidarPointCloudRenderBuffer&) = delete;
-	FLidarPointCloudRenderBuffer(FLidarPointCloudRenderBuffer&&) = delete;
-	FLidarPointCloudRenderBuffer& operator=(const FLidarPointCloudRenderBuffer&) = delete;
-	FLidarPointCloudRenderBuffer& operator=(FLidarPointCloudRenderBuffer&&) = delete;
 };
 
 /**
diff --git a/Engine/Plugins/Enterprise/LidarPointCloud/Source/LidarPointCloudRuntime/Private/Rendering/LidarPointCloudRendering.cpp b/Engine/Plugins/Enterprise/LidarPointCloud/Source/LidarPointCloudRuntime/Private/Rendering/LidarPointCloudRendering.cpp
--- a/Engine/Plugins/Enterprise/LidarPointCloud/Source/LidarPointCloudRuntime/Private/Rendering/LidarPointCloudRendering.cpp
+++ b/Engine/Plugins/Enterprise/LidarPointCloud/Source/LidarPointCloudRuntime/Private/Rendering/LidarPointCloudRendering.cpp
@@ -202,7 +202,7 @@ public:
 					FMeshBatchElement& BatchElement = MeshBatch.Elements[0];
 					BatchElement.PrimitiveUniformBuffer = GetUniformBuffer();
 					BatchElement.IndexBuffer = &GLidarPointCloudIndexBuffer;
-					BatchElement.FirstIndex = bUsesSprites ? 0 : GLidarPointCloudIndexBuffer.GetPointOffset();
+					BatchElement.FirstIndex = bUsesSprites ? 0 : GLidarPointCloudIndexBuffer.PointOffset;
 					BatchElement.MinVertexIndex = 0;
 					BatchElement.NumPrimitives = RenderData.NumElements * (bUsesSprites ? 2 : 1);
 					BatchElement.UserData = &UserData[UserData.Add(BuildUserDataElement(View))];
diff --git a/Engine/Plugins/Enterprise/LidarPointCloud/Source/LidarPointCloudRuntime/Public/LidarPointCloudSettings.h b/Engine/Plugins/Enterprise/LidarPointCloud/Source/LidarPointCloudRuntime/Public/LidarPointCloudSettings.h
--- a/Engine/Plugins/Enterprise/LidarPointCloud/Source/LidarPointCloudRuntime/Public/LidarPointCloudSettings.h
+++ b/Engine/Plugins/Enterprise/LidarPointCloud/Source/LidarPointCloudRuntime/Public/LidarPointCloudSettings.h
@@ -122,10 +122,6 @@ public:
 	UPROPERTY(config, EditAnywhere, Category=Debug)
 	bool bUseIOCaching;
 
-	/** Enabling this will cause Structure and Index buffers to write to the log file anytime they are resized. */
-	UPROPERTY(config, EditAnywhere, Category=Debug)
-	bool bLogBufferExpansion;
-
 public:
 	ULidarPointCloudSettings();
 };
\ No newline at end of file
