Lidar
Fixed Crash when adding a duplicate LiDAR Point Cloud Asset to a level

diff --git a/Engine/Plugins/Enterprise/LidarPointCloud/Source/LidarPointCloudRuntime/Private/LidarPointCloudOctree.cpp b/Engine/Plugins/Enterprise/LidarPointCloud/Source/LidarPointCloudRuntime/Private/LidarPointCloudOctree.cpp
--- a/Engine/Plugins/Enterprise/LidarPointCloud/Source/LidarPointCloudRuntime/Private/LidarPointCloudOctree.cpp
+++ b/Engine/Plugins/Enterprise/LidarPointCloud/Source/LidarPointCloudRuntime/Private/LidarPointCloudOctree.cpp
@@ -1795,6 +1795,7 @@ void FLidarPointCloudOctree::Serialize(FArchive& Ar)
 		Ar.Serialize(CollisionMeshPtr->Indices.GetData(), NumIndices * sizeof(FTriIndices));
 	}
 
+	const bool bIsDuplicating = Ar.GetArchiveName().Equals("FDuplicateDataWriter", ESearchCase::IgnoreCase);
 	const bool bUseCompression = GetDefault<ULidarPointCloudSettings>()->bUseCompression;
 
 	// Used for backwards compatibility with pre-streaming formats
@@ -1842,9 +1843,16 @@ void FLidarPointCloudOctree::Serialize(FArchive& Ar)
 				CurrentNode->SortVisiblePoints();
 			}
 
+			// If preloading for duplication, make sure the data is marked accordingly
+			if (bIsDuplicating)
+			{
+				CurrentNode->GetPersistentData();
+			}
+
 			CurrentNode->BulkData.Serialize(Ar, Owner);
 
-			if (Ar.IsSaving())
+			// Don't reset the release flag if processing duplication
+			if (!bIsDuplicating && Ar.IsSaving())
 			{
 				CurrentNode->bCanReleaseData = true;
 			}
