Fix deprecation warning

diff --git a/Engine/Plugins/Enterprise/LidarPointCloud/Source/LidarPointCloudEditor/Private/LidarPointCloudFactory.cpp b/Engine/Plugins/Enterprise/LidarPointCloud/Source/LidarPointCloudEditor/Private/LidarPointCloudFactory.cpp
--- a/Engine/Plugins/Enterprise/LidarPointCloud/Source/LidarPointCloudEditor/Private/LidarPointCloudFactory.cpp
+++ b/Engine/Plugins/Enterprise/LidarPointCloud/Source/LidarPointCloudEditor/Private/LidarPointCloudFactory.cpp
@@ -118,7 +118,7 @@ void FAssetTypeActions_LidarPointCloud::ExecuteMerge(TArray<ULidarPointCloud*> P
 
 			for (ULidarPointCloud* Asset : PointClouds)
 			{
-				Names.Add(FPaths::GetBaseFilename(FStringAssetReference(Asset).ToString()));
+				Names.Add(FPaths::GetBaseFilename(FSoftObjectPath(Asset).ToString()));
 			}
 
 			FScopedSlowTask ProgressDialog(PointClouds.Num() + 2, LOCTEXT("Merge", "Merging Point Clouds..."));
