Fix shader compile error, due to function name conflicting with global macro

diff --git a/Engine/Plugins/Enterprise/LidarPointCloud/Shaders/Private/LidarPointCloudVertexFactory.ush b/Engine/Plugins/Enterprise/LidarPointCloud/Shaders/Private/LidarPointCloudVertexFactory.ush
--- a/Engine/Plugins/Enterprise/LidarPointCloud/Shaders/Private/LidarPointCloudVertexFactory.ush
+++ b/Engine/Plugins/Enterprise/LidarPointCloud/Shaders/Private/LidarPointCloudVertexFactory.ush
@@ -132,7 +132,7 @@ struct FPointCloudInstanceData
 	half2 Scale;
 };
 
-FPointCloudInstanceData GetInstanceData(uint VertexId)
+FPointCloudInstanceData LidarGetInstanceData(uint VertexId)
 {
 	FPointCloudInstanceData InstanceData = (FPointCloudInstanceData)0;
 
@@ -239,7 +239,7 @@ float3 ProcessPosition(float3 Position, half3 Normal, half2 Scale)
 FVertexFactoryIntermediates GetVertexFactoryIntermediates(FVertexFactoryInput Input)
 {
 	FVertexFactoryIntermediates Intermediates = (FVertexFactoryIntermediates)0;
-	FPointCloudInstanceData InstanceData = GetInstanceData(Input.VertexId);
+	FPointCloudInstanceData InstanceData = LidarGetInstanceData(Input.VertexId);
 	
 	Intermediates.Position = InstanceData.Position;
 	Intermediates.Scale = InstanceData.Scale;
@@ -311,7 +311,7 @@ float4 VertexFactoryGetWorldPosition(FVertexFactoryInput Input, FVertexFactoryIn
 }
 float4 VertexFactoryGetWorldPosition(FPositionOnlyVertexFactoryInput Input)
 {
-	FPointCloudInstanceData InstanceData = GetInstanceData(Input.VertexId);
+	FPointCloudInstanceData InstanceData = LidarGetInstanceData(Input.VertexId);
 
 	float4 WorldPosition = TransformLocalToTranslatedWorld(InstanceData.Position);
 	WorldPosition.xyz = ProcessPosition(WorldPosition.xyz, InstanceData.Normal, InstanceData.Scale);
@@ -323,7 +323,7 @@ float3 VertexFactoryGetPositionForVertexLighting(FVertexFactoryInput Input, FVer
 float3 VertexFactoryGetWorldNormal(FPositionAndNormalOnlyVertexFactoryInput Input) { return RotateLocalToWorld(Input.Normal.xyz, 0); }
 float4 VertexFactoryGetWorldPosition(FPositionAndNormalOnlyVertexFactoryInput Input)
 {
-	FPointCloudInstanceData InstanceData = GetInstanceData(Input.VertexId);
+	FPointCloudInstanceData InstanceData = LidarGetInstanceData(Input.VertexId);
 
 	float4 WorldPosition = TransformLocalToTranslatedWorld(InstanceData.Position, 0);
 	WorldPosition.xyz = ProcessPosition(WorldPosition.xyz, InstanceData.Normal, InstanceData.Scale);
