Lidar: Fixed startup crash

diff --git a/Engine/Plugins/Enterprise/LidarPointCloud/Shaders/Private/LidarPointCloudVertexFactory.ush b/Engine/Plugins/Enterprise/LidarPointCloud/Shaders/Private/LidarPointCloudVertexFactory.ush
--- a/Engine/Plugins/Enterprise/LidarPointCloud/Shaders/Private/LidarPointCloudVertexFactory.ush
+++ b/Engine/Plugins/Enterprise/LidarPointCloud/Shaders/Private/LidarPointCloudVertexFactory.ush
@@ -328,7 +328,7 @@ float4 VertexFactoryGetWorldPosition(FPositionOnlyVertexFactoryInput Input)
 float4 VertexFactoryGetRasterizedWorldPosition(FVertexFactoryInput Input, FVertexFactoryIntermediates Intermediates, float4 InWorldPosition) { return InWorldPosition; }
 float3 VertexFactoryGetPositionForVertexLighting(FVertexFactoryInput Input, FVertexFactoryIntermediates Intermediates, float3 TranslatedWorldPosition) { return TranslatedWorldPosition; }
 
-float3 VertexFactoryGetWorldNormal(FPositionAndNormalOnlyVertexFactoryInput Input) { return RotateLocalToWorld(GetInstanceData(Input.VertexId).TangentToLocal[2], 0); }
+float3 VertexFactoryGetWorldNormal(FPositionAndNormalOnlyVertexFactoryInput Input) { return RotateLocalToWorld(LidarGetInstanceData(Input.VertexId).TangentToLocal[2], 0); }
 float4 VertexFactoryGetWorldPosition(FPositionAndNormalOnlyVertexFactoryInput Input)
 {
 	FPointCloudInstanceData InstanceData = LidarGetInstanceData(Input.VertexId);
