Lidar: Removed UE3 reference

diff --git a/Engine/Plugins/Enterprise/LidarPointCloud/Source/LidarPointCloudEditor/Private/LidarPointCloudEditorViewportClient.cpp b/Engine/Plugins/Enterprise/LidarPointCloud/Source/LidarPointCloudEditor/Private/LidarPointCloudEditorViewportClient.cpp
--- a/Engine/Plugins/Enterprise/LidarPointCloud/Source/LidarPointCloudEditor/Private/LidarPointCloudEditorViewportClient.cpp
+++ b/Engine/Plugins/Enterprise/LidarPointCloud/Source/LidarPointCloudEditor/Private/LidarPointCloudEditorViewportClient.cpp
@@ -280,7 +280,7 @@ bool FLidarPointCloudEditorViewportClient::ShouldOrbitCamera() const
 {
 	if (GetDefault<ULevelEditorViewportSettings>()->bUseUE3OrbitControls)
 	{
-		// this editor orbits always if ue3 orbit controls are enabled
+		// this editor orbits always if orbit controls are enabled
 		return true;
 	}
 
