Lidar: Merging //UE4/Main To //UE5/Release-Engine-Staging

diff --git a/Engine/Plugins/Enterprise/LidarPointCloud/Shaders/Private/LidarPointCloudVertexFactory.ush b/Engine/Plugins/Enterprise/LidarPointCloud/Shaders/Private/LidarPointCloudVertexFactory.ush
--- a/Engine/Plugins/Enterprise/LidarPointCloud/Shaders/Private/LidarPointCloudVertexFactory.ush
+++ b/Engine/Plugins/Enterprise/LidarPointCloud/Shaders/Private/LidarPointCloudVertexFactory.ush
@@ -10,34 +10,30 @@
 // First half is normal, second half has inverted V
 static float2 PrecomputedUV_Quad[8] = { float2(-0.5, -0.5), float2(0.5, -0.5), float2(0.5, 0.5), float2(-0.5, 0.5), float2(0, 1), float2(1, 1), float2(1, 0), float2(0, 0) };
 
-static float3 LODColors[] = {
-	float3(1, 1, 1), float3(1, 0, 0), float3(0, 1, 0), float3(0, 0, 1), float3(1, 1, 0), float3(1, 0, 1), float3(0, 1, 1),
-	float3(0.66, 0.66, 0.66), float3(0.66, 0, 0), float3(0, 0.66, 0), float3(0, 0, 0.66), float3(0.66, 0.66, 0), float3(0.66, 0, 0.66), float3(0, 0.66, 0.66),
-	float3(0.33, 0.33, 0.33), float3(0.33, 0, 0), float3(0, 0.33, 0), float3(0, 0, 0.33), float3(0.33, 0.33, 0), float3(0.33, 0, 0.33), float3(0, 0.33, 0.33)
-};
-
 // Used to detect if the instance is rendered inside the editor viewport
 uint bEditorView;
 
 // Selection
 half3 SelectionColor;
 
-// Used for instance handling
-uint IndexDivisor;
-
 // Alignment
 float3 LocationOffset;
 
-// Used for Sprite Size
+// Used for sprite size calculation
 float VirtualDepth;
-float SpriteSize;
 float SpriteSizeMultiplier;
+float RootCellSize;
+float3 RootExtent;
+uint bUsePerPointScaling;
+float ReversedVirtualDepthMultiplier;
+Buffer<float> TreeBuffer;
 
 // Needed for WPO calculations
 float3 ViewRightVector;
 float3 ViewUpVector;
 uint bUseCameraFacing;
 uint bUseScreenSizeScaling;
+uint bUseStaticBuffers;
 
 // Used for coloration override
 float3 BoundsSize;
@@ -54,7 +50,6 @@ float IntensityInfluence;
 
 // Flags
 uint bUseCircle;
-uint bUseLODColoration;
 uint bUseColorOverride;
 uint bUseElevationColor;
 uint bUseClassification;
@@ -73,7 +68,6 @@ float4x4 ClippingVolume[16];
 uint NumClippingVolumes;
 uint bStartClipped;
 
-
 #if USE_INSTANCING
 	// Instanced stereo sets the eye index explicitly for instanced geometry
 	#if INSTANCED_STEREO
@@ -83,7 +77,9 @@ uint bStartClipped;
 
 struct FVertexFactoryInput
 {
-    float4  Position    : ATTRIBUTE0;
+    float4	Position    : ATTRIBUTE0;
+    half4	Color		: ATTRIBUTE1;
+    uint	MetaData	: ATTRIBUTE2;
 	uint    VertexId    : SV_VertexID;
 };
 
@@ -95,9 +91,9 @@ struct FPositionOnlyVertexFactoryInput
 
 struct FPositionAndNormalOnlyVertexFactoryInput
 {
-	float4  Position : ATTRIBUTE0;
-	float4	Normal : ATTRIBUTE2;
-	uint    VertexId : SV_VertexID;
+	float4  Position	: ATTRIBUTE0;
+	float4	Normal		: ATTRIBUTE2;
+	uint    VertexId	: SV_VertexID;
 };
 
 struct FVertexFactoryInterpolantsVSToPS
@@ -120,6 +116,7 @@ struct FVertexFactoryIntermediates
     half4 Color;
 	float3 Position;
 	half2 Scale;
+	float SpriteSize;
 };
 
 struct FPointCloudInstanceData
@@ -129,47 +126,72 @@ struct FPointCloudInstanceData
 	uint bSelected;
 	uint Classification;
 	half3x3 TangentToLocal;
-	half2 Scale;
+	float SpriteSize;
 };
 
-FPointCloudInstanceData LidarGetInstanceData(uint VertexId)
+float3 ApplyClipping(float3 Position)
 {
-	FPointCloudInstanceData InstanceData = (FPointCloudInstanceData)0;
+	bool bClip = bStartClipped;
 
-	uint PointId = VertexId / IndexDivisor;
-	uint Idx = PointId * 5;
-	InstanceData.Position = float3(DataBuffer[Idx], DataBuffer[Idx + 1], DataBuffer[Idx + 2]) + LocationOffset;
-	InstanceData.Color = asuint(DataBuffer[Idx + 3]);
+	// Only process clipping if not in the editor view
+	if (!bEditorView)
+	{
+		for (uint i = 0; i < NumClippingVolumes; i++)
+		{
+			float3 DeltaPosition = TransformLocalToWorld(Position).xyz - ClippingVolume[i][0].xyz;
+			bool bInsideClippingBox = abs(dot(DeltaPosition, ClippingVolume[i][1].xyz)) <= ClippingVolume[i][1].w && abs(dot(DeltaPosition, ClippingVolume[i][2].xyz)) <= ClippingVolume[i][2].w && abs(dot(DeltaPosition, ClippingVolume[i][3].xyz)) <= ClippingVolume[i][3].w;
 
-	/**
-	 * 0		bVisible
-	 * 1 - 5	Classification
-	 * 6		bSelected
-	 * 7		Unused
-	 */
-	uint MetaData = 0x000000FF & (asuint(DataBuffer[Idx + 4]) >> 24);
+			if (ClippingVolume[i][0][3] == 1)
+			{
+				if (bInsideClippingBox)
+				{
+					bClip = true;
+				}
+			}
+			else
+			{
+				if (bInsideClippingBox)
+				{
+					bClip = false;
+				}
+			}
+		}
+	}
 
-	InstanceData.Classification = 0x0000001F & (MetaData >> 1);
-	InstanceData.bSelected = bEditorView && (0x00000040 & MetaData);
+	return (!bEditorView && bClip) ? 0.0f / 0 : Position;
+}
+
+float3 CalcPosition(float3 Position)
+{
+	return ApplyClipping(Position + LocationOffset);
+}
+
+half3x3 CalcTangentToLocal(uint MetaData)
+{
+	half3x3 TangentToLocal;
 
 	if (bUseCameraFacing)
 	{
-		InstanceData.TangentToLocal[0] = ViewRightVector;
-		InstanceData.TangentToLocal[1] = ViewUpVector;
-		InstanceData.TangentToLocal[2] = cross(InstanceData.TangentToLocal[0], InstanceData.TangentToLocal[1]);
+		// In this mode, [0] and [1] are only really used for scaling
+		TangentToLocal[0] = ViewRightVector;
+		TangentToLocal[1] = ViewUpVector;
+
+		// Hardcoded to have the lighting behave more correctly
+		TangentToLocal[2] = float3(0, 0, 1);
 	}
 	else
 	{
-		uint NormalAsInt = asuint(DataBuffer[Idx + 4]);
-		
-		half3 Normal = half3((0x000000FF & NormalAsInt) / 127.5f - 1, (0x000000FF & (NormalAsInt >> 8)) / 127.5f - 1, (0x000000FF & (NormalAsInt >> 16)) / 127.5f - 1);
-		
+		half3 Normal = half3((0x000000FF & MetaData) / 127.5f - 1, (0x000000FF & (MetaData >> 8)) / 127.5f - 1, (0x000000FF & (MetaData >> 16)) / 127.5f - 1);
+
 		// Force camera facing
 		if (dot(Normal, Normal) < 0.1f)
 		{
-			InstanceData.TangentToLocal[0] = ViewRightVector;
-			InstanceData.TangentToLocal[1] = ViewUpVector;
-			InstanceData.TangentToLocal[2] = cross(InstanceData.TangentToLocal[0], InstanceData.TangentToLocal[1]);
+			// In this mode, [0] and [1] are only really used for scaling
+			TangentToLocal[0] = ViewRightVector;
+			TangentToLocal[1] = ViewUpVector;
+
+			// Hardcoded to have the lighting behave more correctly
+			TangentToLocal[2] = float3(0, 0, 1);
 		}
 		else
 		{
@@ -190,70 +212,151 @@ FPointCloudInstanceData LidarGetInstanceData(uint VertexId)
 			Tangent = (Tangent - Normal * dot(Tangent, Normal));
 			Tangent = Tangent * rsqrt(dot(Tangent, Tangent));
 
-			InstanceData.TangentToLocal[0] = Tangent;
-			InstanceData.TangentToLocal[1] = cross(Tangent, Normal);
-			InstanceData.TangentToLocal[2] = Normal;
+			TangentToLocal[0] = Tangent;
+			TangentToLocal[1] = cross(Tangent, Normal);
+			TangentToLocal[2] = Normal;
 		}
 	}
 
-	InstanceData.Scale = PrecomputedUV_Quad[VertexId % 4];
+	return TangentToLocal;
+}
 
-	// Only process clipping if not in the editor view
-	if (!bEditorView)
+half3x3 CalcTangentToWorldNoScale(in half3x3 TangentToLocal)
+{
+	half3x3 LocalToWorld = GetLocalToWorld3x3();
+	half3 InvScale = Primitive.InvNonUniformScaleAndDeterminantSign.xyz;
+	LocalToWorld[0] *= InvScale.x;
+	LocalToWorld[1] *= InvScale.y;
+	LocalToWorld[2] *= InvScale.z;
+	return mul(TangentToLocal, LocalToWorld);
+}
+
+float CalcSpriteSize(float3 Position)
+{
+	float VD = VirtualDepth;
+
+	if (bUsePerPointScaling)
 	{
-		bool bClip = bStartClipped;
-		for (uint i = 0; i < NumClippingVolumes; i++)
+		VD = 0;
+		uint Idx = 0;
+		float3 NodeExtent = RootExtent;
+		float3 NodeCenter = LocationOffset;
+
+		while (true)
 		{
-			float3 DeltaPosition = TransformLocalToWorld(InstanceData.Position).xyz - ClippingVolume[i][0].xyz;
-			bool bInsideClippingBox = abs(dot(DeltaPosition, ClippingVolume[i][1].xyz)) <= ClippingVolume[i][1].w && abs(dot(DeltaPosition, ClippingVolume[i][2].xyz)) <= ClippingVolume[i][2].w && abs(dot(DeltaPosition, ClippingVolume[i][3].xyz)) <= ClippingVolume[i][3].w;
-		
-			if (ClippingVolume[i][0][3] == 1)
+			uint TreeBufferData = asuint(TreeBuffer[Idx]);
+			uint ChildrenBitmask = 0x000000FF & TreeBufferData;
+
+			float3 CenterRelativeLocation = Position - NodeCenter;
+			uint ChildNodeIndex = (CenterRelativeLocation.x > 0 ? 4 : 0) + (CenterRelativeLocation.y > 0 ? 2 : 0) + (CenterRelativeLocation.z > 0);
+
+			if ((ChildrenBitmask & (1 << ChildNodeIndex)) == 0)
 			{
-				if (bInsideClippingBox)
-				{
-					bClip = true;
-				}
+				break;
 			}
-			else
+
+			Idx += (0x0000FFFF & (TreeBufferData >> 8));
+			for (int i = 0; i < 8; i++)
 			{
-				if (bInsideClippingBox)
+				if (i == ChildNodeIndex)
 				{
-					bClip = false;
+					break;
+				}
+
+				if ((ChildrenBitmask & (1 << i)) != 0)
+				{
+					Idx++;
 				}
 			}
-		}
 
-		InstanceData.Position = bClip ? 0.0f / 0 : InstanceData.Position;
+			VD = (0x000000FF & (TreeBufferData >> 24)) * ReversedVirtualDepthMultiplier;
+			NodeCenter += NodeExtent * (float3((ChildNodeIndex & 4) == 4, (ChildNodeIndex & 2) == 2, (ChildNodeIndex & 1) == 1) - 0.5f);
+			NodeExtent *= 0.5f;
+		}
 	}
 	
-	return InstanceData;
+	return RootCellSize / pow(2.0f, VD);
 }
 
-half3x3 CalcTangentToWorldNoScale(in half3x3 TangentToLocal)
+float3 ProcessPosition(float3 Position, half3x3 TangentToLocal, half2 Scale, float SpriteSize)
 {
-	half3x3 LocalToWorld = GetLocalToWorld3x3();
-	half3 InvScale = Primitive.InvNonUniformScaleAndDeterminantSign.xyz;
-	LocalToWorld[0] *= InvScale.x;
-	LocalToWorld[1] *= InvScale.y;
-	LocalToWorld[2] *= InvScale.z;
-	return mul(TangentToLocal, LocalToWorld);
+	float Size = bUseScreenSizeScaling ? mul(float4(Position.xyz, 1), ResolvedView.TranslatedWorldToView).z * 0.01f : SpriteSize;
+	return Position + (TangentToLocal[0] * Scale.x + TangentToLocal[1] * Scale.y) * SpriteSizeMultiplier * Size;
 }
 
-float3 ProcessPosition(float3 Position, half3x3 TangentToLocal, half2 Scale)
+float3 ProcessPosition(float3 Position, half2 Scale)
 {
-	float Size = bUseScreenSizeScaling ? mul(float4(Position.xyz, 1), ResolvedView.TranslatedWorldToView).z * 0.01f : SpriteSize;
-	return Position + (TangentToLocal[0] * Scale.x + TangentToLocal[1] * Scale.y) * SpriteSizeMultiplier * Size;
+	half3x3 TangentToLocal;
+	TangentToLocal[0] = ViewRightVector;
+	TangentToLocal[1] = ViewUpVector;
+	return ProcessPosition(Position, TangentToLocal, Scale, CalcSpriteSize(Position));
+}
+
+FPointCloudInstanceData LidarGetInstanceData(uint VertexId)
+{
+	FPointCloudInstanceData InstanceData = (FPointCloudInstanceData)0;
+
+	uint PointId = VertexId / 4;
+	uint Idx = PointId * 5;
+	InstanceData.Position = float3(DataBuffer[Idx], DataBuffer[Idx + 1], DataBuffer[Idx + 2]) + LocationOffset;
+	InstanceData.Color = asuint(DataBuffer[Idx + 3]);
+
+	/**
+	 * 0		bVisible
+	 * 1 - 5	Classification
+	 * 6		bSelected
+	 * 7		Unused
+	 */
+	uint MetaData = 0x000000FF & (asuint(DataBuffer[Idx + 4]) >> 24);
+
+	InstanceData.Classification = 0x0000001F & (MetaData >> 1);
+	InstanceData.bSelected = bEditorView && (0x00000040 & MetaData);
+	InstanceData.TangentToLocal = CalcTangentToLocal(asuint(DataBuffer[Idx + 4]));
+	InstanceData.Position = ApplyClipping(InstanceData.Position);
+	InstanceData.SpriteSize = CalcSpriteSize(InstanceData.Position);
+
+	return InstanceData;
 }
 
 FVertexFactoryIntermediates GetVertexFactoryIntermediates(FVertexFactoryInput Input)
 {
 	FVertexFactoryIntermediates Intermediates = (FVertexFactoryIntermediates)0;
-	FPointCloudInstanceData InstanceData = LidarGetInstanceData(Input.VertexId);
 	
-	Intermediates.Position = InstanceData.Position;
-	Intermediates.Scale = InstanceData.Scale;
+	uint bSelected;
+	uint Classification;
+	half3 Color;
+	half Intensity;
+	float SpriteSize;
+
+	if (bUseStaticBuffers)
+	{
+		Intermediates.Position = CalcPosition(Input.Position.xyz);
+		Intermediates.TangentToLocal = CalcTangentToLocal(Input.MetaData);
+
+		uint MetaData = 0x000000FF & (Input.MetaData >> 24);
 
-	Intermediates.TangentToLocal = InstanceData.TangentToLocal;
+		bSelected = bEditorView && (0x00000040 & MetaData);
+		Classification = 0x0000001F & (MetaData >> 1);
+
+		SpriteSize = CalcSpriteSize(Intermediates.Position);
+
+		Color = Input.Color.rgb;
+		Intensity = Input.Color.a;
+	}
+	else
+	{
+		FPointCloudInstanceData InstanceData = LidarGetInstanceData(Input.VertexId);
+		Intermediates.Position = InstanceData.Position;
+		Intermediates.TangentToLocal = InstanceData.TangentToLocal;
+		bSelected = InstanceData.bSelected;
+		Classification = InstanceData.Classification;
+		SpriteSize = InstanceData.SpriteSize;
+		Color = half3(0x000000FF & (InstanceData.Color >> 16), 0x000000FF & (InstanceData.Color >> 8), 0x000000FF & InstanceData.Color) * 0.003921568627451;
+		Intensity = (0x000000FF & (InstanceData.Color >> 24)) * 0.003921568627451;
+	}
+
+	Intermediates.SpriteSize = SpriteSize;
+	Intermediates.Scale = PrecomputedUV_Quad[Input.VertexId % 4];
 	Intermediates.TangentToWorldSign = Primitive.InvNonUniformScaleAndDeterminantSign.w;
 	Intermediates.TangentToWorld = CalcTangentToWorldNoScale(Intermediates.TangentToLocal);
 	
@@ -262,20 +365,15 @@ FVertexFactoryIntermediates GetVertexFactoryIntermediates(FVertexFactoryInput In
 		// Used to account for offset due to shifted points bounds
 		float3 RelativePosition = Intermediates.Position - LocationOffset;
 
-		half3 Color = half3(0x000000FF & (InstanceData.Color >> 16), 0x000000FF & (InstanceData.Color >> 8), 0x000000FF & InstanceData.Color) * 0.003921568627451;
 		half3 ElevationColor = lerp(ElevationColorBottom, ElevationColorTop, clamp((RelativePosition.z + BoundsSize.z * 0.5) / BoundsSize.z, 0, 1));
 		half3 PositionColor = (RelativePosition + BoundsSize * 0.5) / BoundsSize;
-		half3 ClassificationColor = ClassificationColors[InstanceData.Classification];
+		half3 ClassificationColor = ClassificationColors[0x0000001F & (Classification >> 1)].rgb;
 
 		// Pick color source
-		Color = lerp(lerp(Color, lerp(lerp(PositionColor, ClassificationColor, bUseClassification), ElevationColor, bUseElevationColor), bUseColorOverride), SelectionColor, InstanceData.bSelected);
+		Color = lerp(lerp(Color, lerp(lerp(PositionColor, ClassificationColor, bUseClassification), ElevationColor, bUseElevationColor), bUseColorOverride), SelectionColor, bSelected);
 
 		// Mix with Intensity
-		{
-			half Intensity = (0x000000FF & (InstanceData.Color >> 24)) * 0.003921568627451;
-			Color = lerp(Color, Color * Intensity, IntensityInfluence);
-		}
-
+		Color = lerp(Color, Color * Intensity, IntensityInfluence);
 
 		// Color Adjustment
 		{
@@ -287,7 +385,7 @@ FVertexFactoryIntermediates GetVertexFactoryIntermediates(FVertexFactoryInput In
 			Color = clamp(Color, 0, 1);
 		}
 
-		Intermediates.Color = half4(lerp(Color, LODColors[min((int)VirtualDepth, 21)], bUseLODColoration), bUseCircle);
+		Intermediates.Color = half4(Color, bUseCircle);
 	}
 
 	return Intermediates;
@@ -314,28 +412,46 @@ half3x3 VertexFactoryGetTangentToLocal(FVertexFactoryInput Input, FVertexFactory
 float4 VertexFactoryGetWorldPosition(FVertexFactoryInput Input, FVertexFactoryIntermediates Intermediates)
 {
 	float4 WorldPosition = TransformLocalToTranslatedWorld(Intermediates.Position);
-	WorldPosition.xyz = ProcessPosition(WorldPosition.xyz, Intermediates.TangentToLocal, Intermediates.Scale);
+	WorldPosition.xyz = ProcessPosition(WorldPosition.xyz, Intermediates.TangentToLocal, Intermediates.Scale, Intermediates.SpriteSize);
 	return WorldPosition;
 }
 float4 VertexFactoryGetWorldPosition(FPositionOnlyVertexFactoryInput Input)
 {
-	FPointCloudInstanceData InstanceData = LidarGetInstanceData(Input.VertexId);
-
-	float4 WorldPosition = TransformLocalToTranslatedWorld(InstanceData.Position);
-	WorldPosition.xyz = ProcessPosition(WorldPosition.xyz, InstanceData.TangentToLocal, InstanceData.Scale);
-	return WorldPosition;
+	if (bUseStaticBuffers)
+	{
+		float4 WorldPosition = TransformLocalToTranslatedWorld(CalcPosition(Input.Position.xyz));
+		WorldPosition.xyz = ProcessPosition(WorldPosition.xyz, PrecomputedUV_Quad[Input.VertexId % 4]);
+		return WorldPosition;
+	}
+	else
+	{
+		FPointCloudInstanceData InstanceData = LidarGetInstanceData(Input.VertexId);
+		float4 WorldPosition = TransformLocalToTranslatedWorld(InstanceData.Position);
+		WorldPosition.xyz = ProcessPosition(WorldPosition.xyz, InstanceData.TangentToLocal, PrecomputedUV_Quad[Input.VertexId % 4], InstanceData.SpriteSize);
+		return WorldPosition;
+	}
 }
 float4 VertexFactoryGetRasterizedWorldPosition(FVertexFactoryInput Input, FVertexFactoryIntermediates Intermediates, float4 InWorldPosition) { return InWorldPosition; }
 float3 VertexFactoryGetPositionForVertexLighting(FVertexFactoryInput Input, FVertexFactoryIntermediates Intermediates, float3 TranslatedWorldPosition) { return TranslatedWorldPosition; }
-
-float3 VertexFactoryGetWorldNormal(FPositionAndNormalOnlyVertexFactoryInput Input) { return RotateLocalToWorld(LidarGetInstanceData(Input.VertexId).TangentToLocal[2], 0); }
+float3 VertexFactoryGetWorldNormal(FPositionAndNormalOnlyVertexFactoryInput Input)
+{
+	return RotateLocalToWorld(bUseStaticBuffers ? Input.Normal.xyz : LidarGetInstanceData(Input.VertexId).TangentToLocal[2], 0);
+}
 float4 VertexFactoryGetWorldPosition(FPositionAndNormalOnlyVertexFactoryInput Input)
 {
-	FPointCloudInstanceData InstanceData = LidarGetInstanceData(Input.VertexId);
-
-	float4 WorldPosition = TransformLocalToTranslatedWorld(InstanceData.Position, 0);
-	WorldPosition.xyz = ProcessPosition(WorldPosition.xyz, InstanceData.TangentToLocal, InstanceData.Scale);
-	return WorldPosition;
+	if (bUseStaticBuffers)
+	{
+		float4 WorldPosition = TransformLocalToTranslatedWorld(CalcPosition(Input.Position.xyz), 0);
+		WorldPosition.xyz = ProcessPosition(WorldPosition.xyz, PrecomputedUV_Quad[Input.VertexId % 4]);
+		return WorldPosition;
+	}
+	else
+	{
+		FPointCloudInstanceData InstanceData = LidarGetInstanceData(Input.VertexId);
+		float4 WorldPosition = TransformLocalToTranslatedWorld(InstanceData.Position, 0);
+		WorldPosition.xyz = ProcessPosition(WorldPosition.xyz, InstanceData.TangentToLocal, PrecomputedUV_Quad[Input.VertexId % 4], InstanceData.SpriteSize);
+		return WorldPosition;
+	}
 }
 float3 VertexFactoryGetWorldNormal(FVertexFactoryInput Input, FVertexFactoryIntermediates Intermediates) { return Intermediates.TangentToWorld[2]; }
 
@@ -466,3 +582,4 @@ uint VertexFactoryGetPrimitiveId(FVertexFactoryInterpolantsVSToPS Interpolants)
 }
 
 #include "/Engine/Private/VertexFactoryDefaultInterface.ush"
+
diff --git a/Engine/Plugins/Enterprise/LidarPointCloud/Source/LidarPointCloudEditor/Private/LidarPointCloudEditor.cpp b/Engine/Plugins/Enterprise/LidarPointCloud/Source/LidarPointCloudEditor/Private/LidarPointCloudEditor.cpp
--- a/Engine/Plugins/Enterprise/LidarPointCloud/Source/LidarPointCloudEditor/Private/LidarPointCloudEditor.cpp
+++ b/Engine/Plugins/Enterprise/LidarPointCloud/Source/LidarPointCloudEditor/Private/LidarPointCloudEditor.cpp
@@ -2,10 +2,10 @@
 
 #include "LidarPointCloudEditor.h"
 #include "LidarPointCloudEditorViewport.h"
+#include "LidarPointCloudEditorViewportClient.h"
 #include "LidarPointCloud.h"
 #include "LidarPointCloudShared.h"
 #include "ILidarPointCloudEditorModule.h"
-
 #include "LidarPointCloudEditorCommands.h"
 
 #include "Misc/ScopedSlowTask.h"
@@ -28,6 +28,45 @@ const FName PointCloudEditorAppName = FName(TEXT("LidarPointCloudEditorApp"));
 const FName FLidarPointCloudEditor::DetailsTabId(TEXT("Details"));
 const FName FLidarPointCloudEditor::ViewportTabId(TEXT("Viewport"));
 
+void SetSelectionMethod(FLidarPointCloudEditor* Editor, ELidarPointCloudSelectionMethod NewSelectionMethod)
+{
+	if (SLidarPointCloudEditorViewport* ViewportPtr = Editor->GetViewport().Get())
+	{
+		if (FLidarPointCloudEditorViewportClient* ViewportClientPtr = ViewportPtr->GetEditorViewportClient().Get())
+		{
+			ViewportClientPtr->SetSelectionMethod(NewSelectionMethod);
+		}
+	}
+
+	Editor->RegenerateMenusAndToolbars();
+}
+
+FName GetSelectionIcon(FLidarPointCloudEditor* Editor)
+{
+	if (SLidarPointCloudEditorViewport* ViewportPtr = Editor->GetViewport().Get())
+	{
+		if (FLidarPointCloudEditorViewportClient* ViewportClientPtr = ViewportPtr->GetEditorViewportClient().Get())
+		{
+			switch (ViewportClientPtr->GetSelectionMethod())
+			{
+			case ELidarPointCloudSelectionMethod::Box:
+				return "LidarPointCloudEditor.BoxSelection";
+
+			case ELidarPointCloudSelectionMethod::Polygonal:
+				return "LidarPointCloudEditor.PolygonalSelection";
+
+			case ELidarPointCloudSelectionMethod::Lasso:
+				return "LidarPointCloudEditor.LassoSelection";
+
+			case ELidarPointCloudSelectionMethod::Paint:
+				return "LidarPointCloudEditor.PaintSelection";
+			}
+		}
+	}
+
+	return "";
+}
+
 /////////////////////////////////////////////////////
 // SPointCloudPropertiesTabBody
 
@@ -188,12 +227,12 @@ void FLidarPointCloudEditor::InitPointCloudEditor(const EToolkitMode::Type Mode,
 	RegenerateMenusAndToolbars();
 }
 
-void FLidarPointCloudEditor::SelectPointsByFrustum(const FConvexVolume& SelectionFrustum, bool bAdditive)
+void FLidarPointCloudEditor::SelectPointsByConvexVolume(const FConvexVolume& ConvexVolume, bool bAdditive)
 {
 	if (bAdditive && SelectedPoints.Num() > 0)
 	{
-		TArray<FLidarPointCloudPoint*> NewSelectedPoints;
-		PointCloudBeingEdited->GetPointsInFrustum(NewSelectedPoints, SelectionFrustum, true);
+		TArray64<FLidarPointCloudPoint*> NewSelectedPoints;
+		PointCloudBeingEdited->GetPointsInConvexVolume(NewSelectedPoints, ConvexVolume, true);
 
 		// Merge selections and set selection flag for new selection of points
 		for (FLidarPointCloudPoint** Data = NewSelectedPoints.GetData(), **DataEnd = Data + NewSelectedPoints.Num(); Data != DataEnd; ++Data)
@@ -210,7 +249,7 @@ void FLidarPointCloudEditor::SelectPointsByFrustum(const FConvexVolume& Selectio
 		// Clear any existing selection
 		DeselectPoints();
 
-		PointCloudBeingEdited->GetPointsInFrustum(SelectedPoints, SelectionFrustum, true);
+		PointCloudBeingEdited->GetPointsInConvexVolume(SelectedPoints, ConvexVolume, true);
 
 		// Set selection flag for new selection of points
 		for (FLidarPointCloudPoint** Data = SelectedPoints.GetData(), **DataEnd = Data + SelectedPoints.Num(); Data != DataEnd; ++Data)
@@ -219,13 +258,13 @@ void FLidarPointCloudEditor::SelectPointsByFrustum(const FConvexVolume& Selectio
 		}
 	}	
 
-	PointCloudBeingEdited->Octree.MarkRenderDataInFrustumDirty(SelectionFrustum);
+	PointCloudBeingEdited->Octree.MarkRenderDataInConvexVolumeDirty(ConvexVolume);
 }
 
-void FLidarPointCloudEditor::DeselectPointsByFrustum(const FConvexVolume& SelectionFrustum)
+void FLidarPointCloudEditor::DeselectPointsByConvexVolume(const FConvexVolume& ConvexVolume)
 {
-	TArray<FLidarPointCloudPoint*> NewSelectedPoints;
-	PointCloudBeingEdited->GetPointsInFrustum(NewSelectedPoints, SelectionFrustum, true);
+	TArray64<FLidarPointCloudPoint*> NewSelectedPoints;
+	PointCloudBeingEdited->GetPointsInConvexVolume(NewSelectedPoints, ConvexVolume, true);
 
 	// Unset selection flag for selection of points
 	for (FLidarPointCloudPoint** Data = NewSelectedPoints.GetData(), **DataEnd = Data + NewSelectedPoints.Num(); Data != DataEnd; ++Data)
@@ -246,7 +285,57 @@ void FLidarPointCloudEditor::DeselectPointsByFrustum(const FConvexVolume& Select
 
 	SelectedPoints.Shrink();
 
-	PointCloudBeingEdited->Octree.MarkRenderDataInFrustumDirty(SelectionFrustum);
+	PointCloudBeingEdited->Octree.MarkRenderDataInConvexVolumeDirty(ConvexVolume);
+}
+
+void FLidarPointCloudEditor::SelectPointsBySphere(const FSphere& Sphere)
+{
+	TArray64<FLidarPointCloudPoint*> NewSelectedPoints;
+	PointCloudBeingEdited->GetPointsInSphere(NewSelectedPoints, Sphere, true);
+
+	// Reserve space
+	const int64 NeededSlack = FMath::Max(0LL, NewSelectedPoints.Num() - SelectedPoints.GetSlack());
+	if (NeededSlack > 0)
+	{
+		SelectedPoints.Reserve(NeededSlack);
+	}
+
+	// Merge selections and set selection flag for new selection of points
+	for (FLidarPointCloudPoint** Data = NewSelectedPoints.GetData(), **DataEnd = Data + NewSelectedPoints.Num(); Data != DataEnd; ++Data)
+	{
+		if (!(*Data)->bSelected)
+		{
+			(*Data)->bSelected = true;
+			SelectedPoints.Add(*Data);
+		}
+	}
+
+	PointCloudBeingEdited->Octree.MarkRenderDataInSphereDirty(Sphere);
+}
+
+void FLidarPointCloudEditor::DeselectPointsBySphere(const FSphere& Sphere)
+{
+	TArray64<FLidarPointCloudPoint*> NewSelectedPoints;
+	PointCloudBeingEdited->GetPointsInSphere(NewSelectedPoints, Sphere, true);
+
+	// Unset selection flag for selection of points
+	for (FLidarPointCloudPoint** Data = NewSelectedPoints.GetData(), **DataEnd = Data + NewSelectedPoints.Num(); Data != DataEnd; ++Data)
+	{
+		(*Data)->bSelected = false;
+	}
+
+	// Remove deselected points from the selection list
+	FLidarPointCloudPoint** Data = SelectedPoints.GetData();
+	for (int32 i = 0; i < SelectedPoints.Num(); ++Data, ++i)
+	{
+		if (!(*Data)->bSelected)
+		{
+			SelectedPoints.RemoveAtSwap(i--, 1, false);
+			--Data;
+		}
+	}
+
+	PointCloudBeingEdited->Octree.MarkRenderDataInSphereDirty(Sphere);
 }
 
 void FLidarPointCloudEditor::DeselectPoints()
@@ -475,6 +564,28 @@ TSharedRef<SWidget> FLidarPointCloudEditor::GenerateDeleteMenuContent()
 	return MenuBuilder.MakeWidget();
 }
 
+TSharedRef<SWidget> FLidarPointCloudEditor::GenerateSelectionMenuContent()
+{
+	FMenuBuilder MenuBuilder(true, Viewport->GetCommandList());
+
+	MenuBuilder.BeginSection("SelectionMethod", LOCTEXT("SelectionMethod", "Selection Method"));
+	{
+		MenuBuilder.AddMenuEntry(FLidarPointCloudEditorCommands::Get().BoxSelection);
+		MenuBuilder.AddMenuEntry(FLidarPointCloudEditorCommands::Get().PolygonalSelection);
+		MenuBuilder.AddMenuEntry(FLidarPointCloudEditorCommands::Get().LassoSelection);
+		MenuBuilder.AddMenuEntry(FLidarPointCloudEditorCommands::Get().PaintSelection);
+	}
+	MenuBuilder.EndSection();
+
+	MenuBuilder.BeginSection("SelectionOptions", LOCTEXT("SelectionOptions", "Selection Options"));
+	{
+		MenuBuilder.AddMenuEntry(FLidarPointCloudEditorCommands::Get().InvertSelection);
+	}
+	MenuBuilder.EndSection();
+
+	return MenuBuilder.MakeWidget();
+}
+
 void FLidarPointCloudEditor::ExtendToolBar()
 {
 	struct Local
@@ -507,7 +618,13 @@ void FLidarPointCloudEditor::ExtendToolBar()
 			ToolbarBuilder.BeginSection("LidarPointCloudEdit");
 			{
 				ToolbarBuilder.AddToolBarButton(Commands->EditMode);
-				ToolbarBuilder.AddToolBarButton(Commands->InvertSelection);
+				ToolbarBuilder.AddComboButton(
+					FUIAction(FExecuteAction(), FCanExecuteAction::CreateSP(ThisEditor, &FLidarPointCloudEditor::IsEditMode)),
+					FOnGetContent::CreateSP(ThisEditor, &FLidarPointCloudEditor::GenerateSelectionMenuContent),
+					LOCTEXT("Selection_Label", "Selection"),
+					LOCTEXT("Selection_Tooltip", "Selection options"),
+					FSlateIcon(FLidarPointCloudStyle::GetStyleSetName(), GetSelectionIcon(ThisEditor))
+				);			
 				ToolbarBuilder.AddToolBarButton(Commands->HideSelected);
 				ToolbarBuilder.AddToolBarButton(Commands->UnhideAll);
 				ToolbarBuilder.AddComboButton(
@@ -553,6 +670,10 @@ void FLidarPointCloudEditor::BindEditorCommands()
 	CommandList->MapAction(Commands.BuildCollision, FExecuteAction::CreateSP(this, &FLidarPointCloudEditor::BuildCollision));
 	CommandList->MapAction(Commands.RemoveCollision, FExecuteAction::CreateSP(this, &FLidarPointCloudEditor::RemoveCollision), FCanExecuteAction::CreateSP(this, &FLidarPointCloudEditor::HasCollisionData));
 	CommandList->MapAction(Commands.EditMode, FExecuteAction::CreateSP(this, &FLidarPointCloudEditor::ToggleEditMode), FCanExecuteAction(), FIsActionChecked::CreateSP(this, &FLidarPointCloudEditor::IsEditMode));
+	CommandList->MapAction(Commands.BoxSelection, FExecuteAction::CreateLambda([this]{ SetSelectionMethod(this, ELidarPointCloudSelectionMethod::Box); }), FCanExecuteAction::CreateSP(this, &FLidarPointCloudEditor::IsEditMode));
+	CommandList->MapAction(Commands.PolygonalSelection, FExecuteAction::CreateLambda([this]{ SetSelectionMethod(this, ELidarPointCloudSelectionMethod::Polygonal); }), FCanExecuteAction::CreateSP(this, &FLidarPointCloudEditor::IsEditMode));
+	CommandList->MapAction(Commands.LassoSelection, FExecuteAction::CreateLambda([this] { SetSelectionMethod(this, ELidarPointCloudSelectionMethod::Lasso); }), FCanExecuteAction::CreateSP(this, &FLidarPointCloudEditor::IsEditMode));
+	CommandList->MapAction(Commands.PaintSelection, FExecuteAction::CreateLambda([this] { SetSelectionMethod(this, ELidarPointCloudSelectionMethod::Paint); }), FCanExecuteAction::CreateSP(this, &FLidarPointCloudEditor::IsEditMode));
 	CommandList->MapAction(Commands.InvertSelection, FExecuteAction::CreateSP(this, &FLidarPointCloudEditor::InvertSelection), FCanExecuteAction::CreateSP(this, &FLidarPointCloudEditor::HasSelectedPoints));
 	CommandList->MapAction(Commands.UnhideAll, FExecuteAction::CreateSP(this, &FLidarPointCloudEditor::UnhideAll));
 	CommandList->MapAction(Commands.HideSelected, FExecuteAction::CreateSP(this, &FLidarPointCloudEditor::HidePoints), FCanExecuteAction::CreateSP(this, &FLidarPointCloudEditor::HasSelectedPoints));
diff --git a/Engine/Plugins/Enterprise/LidarPointCloud/Source/LidarPointCloudEditor/Private/LidarPointCloudEditor.h b/Engine/Plugins/Enterprise/LidarPointCloud/Source/LidarPointCloudEditor/Private/LidarPointCloudEditor.h
--- a/Engine/Plugins/Enterprise/LidarPointCloud/Source/LidarPointCloudEditor/Private/LidarPointCloudEditor.h
+++ b/Engine/Plugins/Enterprise/LidarPointCloud/Source/LidarPointCloudEditor/Private/LidarPointCloudEditor.h
@@ -9,10 +9,12 @@
 #include "Widgets/Text/STextBlock.h"
 #include "AssetData.h"
 #include "ConvexVolume.h"
+#include "LidarPointCloudShared.h"
 
 class ULidarPointCloud;
 struct FLidarPointCloudPoint;
 class SWidget;
+class SLidarPointCloudEditorViewport;
 
 class FLidarPointCloudEditor : public FAssetEditorToolkit, public FGCObject
 {
@@ -24,7 +26,7 @@ private:
 	bool bEditMode;
 
 	/** Preview Viewport widget */
-	TSharedPtr<class SLidarPointCloudEditorViewport> Viewport;
+	TSharedPtr<SLidarPointCloudEditorViewport> Viewport;
 
 public:
 	FLidarPointCloudEditor();
@@ -59,8 +61,10 @@ public:
 
 	ULidarPointCloud* GetPointCloudBeingEdited() const { return PointCloudBeingEdited; }
 
-	void SelectPointsByFrustum(const FConvexVolume& SelectionFrustum, bool bAdditive);
-	void DeselectPointsByFrustum(const FConvexVolume& SelectionFrustum);
+	void SelectPointsByConvexVolume(const FConvexVolume& ConvexVolume, bool bAdditive);
+	void DeselectPointsByConvexVolume(const FConvexVolume& ConvexVolume);
+	void SelectPointsBySphere(const FSphere& Sphere);
+	void DeselectPointsBySphere(const FSphere& Sphere);
 	void DeselectPoints();
 	void InvertSelection();
 	void DeletePoints();
@@ -73,6 +77,8 @@ public:
 
 	bool IsEditMode() const { return bEditMode; }
 
+	TSharedPtr<SLidarPointCloudEditorViewport> GetViewport() { return Viewport; }
+
 private:
 	bool ConfirmCollisionChange();
 
@@ -108,6 +114,7 @@ private:
 	TSharedRef<SWidget> GenerateExtractionMenuContent();
 	TSharedRef<SWidget> GenerateCollisionMenuContent();
 	TSharedRef<SWidget> GenerateDeleteMenuContent();
+	TSharedRef<SWidget> GenerateSelectionMenuContent();
 
 	TArray<FAssetData> SelectAssets(const FText& Title);
 	FString GetSaveAsLocation();
diff --git a/Engine/Plugins/Enterprise/LidarPointCloud/Source/LidarPointCloudEditor/Private/LidarPointCloudEditorCommands.cpp b/Engine/Plugins/Enterprise/LidarPointCloud/Source/LidarPointCloudEditor/Private/LidarPointCloudEditorCommands.cpp
--- a/Engine/Plugins/Enterprise/LidarPointCloud/Source/LidarPointCloudEditor/Private/LidarPointCloudEditorCommands.cpp
+++ b/Engine/Plugins/Enterprise/LidarPointCloud/Source/LidarPointCloudEditor/Private/LidarPointCloudEditorCommands.cpp
@@ -22,6 +22,11 @@ void FLidarPointCloudEditorCommands::RegisterCommands()
 
 	UI_COMMAND(EditMode, "Edit Mode", "Enables editing of the point cloud.", EUserInterfaceActionType::ToggleButton, FInputChord());
 
+	UI_COMMAND(BoxSelection, "Box Selection", "Uses box to select points.", EUserInterfaceActionType::Button, FInputChord());
+	UI_COMMAND(PolygonalSelection, "Polygonal Selection", "Uses custom polygon to select points.", EUserInterfaceActionType::Button, FInputChord());
+	UI_COMMAND(LassoSelection, "Lasso Selection", "Uses custom drawn shape to select points.", EUserInterfaceActionType::Button, FInputChord());
+	UI_COMMAND(PaintSelection, "Paint Selection", "Uses adjustable paint brush to select points.", EUserInterfaceActionType::Button, FInputChord());
+
 	UI_COMMAND(InvertSelection, "Invert Selection", "Inverts point selection.", EUserInterfaceActionType::Button, FInputChord());
 
 	UI_COMMAND(HideSelected, "Hide Selected", "Hide selected points.", EUserInterfaceActionType::Button, FInputChord());
diff --git a/Engine/Plugins/Enterprise/LidarPointCloud/Source/LidarPointCloudEditor/Private/LidarPointCloudEditorCommands.h b/Engine/Plugins/Enterprise/LidarPointCloud/Source/LidarPointCloudEditor/Private/LidarPointCloudEditorCommands.h
--- a/Engine/Plugins/Enterprise/LidarPointCloud/Source/LidarPointCloudEditor/Private/LidarPointCloudEditorCommands.h
+++ b/Engine/Plugins/Enterprise/LidarPointCloud/Source/LidarPointCloudEditor/Private/LidarPointCloudEditorCommands.h
@@ -38,6 +38,11 @@ public:
 
 	TSharedPtr< FUICommandInfo > EditMode;
 
+	TSharedPtr< FUICommandInfo > BoxSelection;
+	TSharedPtr< FUICommandInfo > PolygonalSelection;
+	TSharedPtr< FUICommandInfo > LassoSelection;
+	TSharedPtr< FUICommandInfo > PaintSelection;
+
 	TSharedPtr< FUICommandInfo > InvertSelection;
 
 	TSharedPtr< FUICommandInfo > HideSelected;
diff --git a/Engine/Plugins/Enterprise/LidarPointCloud/Source/LidarPointCloudEditor/Private/LidarPointCloudEditorViewport.cpp b/Engine/Plugins/Enterprise/LidarPointCloud/Source/LidarPointCloudEditor/Private/LidarPointCloudEditorViewport.cpp
--- a/Engine/Plugins/Enterprise/LidarPointCloud/Source/LidarPointCloudEditor/Private/LidarPointCloudEditorViewport.cpp
+++ b/Engine/Plugins/Enterprise/LidarPointCloud/Source/LidarPointCloudEditor/Private/LidarPointCloudEditorViewport.cpp
@@ -10,6 +10,7 @@
 #include "EditorStyleSet.h"
 #include "ComponentReregisterContext.h"
 #include "Widgets/Docking/SDockTab.h"
+#include "Components/StaticMeshComponent.h"
 
 ///////////////////////////////////////////////////////////
 // SPointCloudEditorViewportToolbar
@@ -73,12 +74,29 @@ void SLidarPointCloudEditorViewport::Construct(const FArguments& InArgs)
 
 	SetPreviewCloud(PointCloud);
 
+	// Add Paint Brush
+	static UStaticMesh* PaintBrushMesh = Cast<UStaticMesh>(StaticLoadObject(UStaticMesh::StaticClass(), nullptr, TEXT("/Engine/BasicShapes/Sphere.Sphere")));
+	static UMaterialInterface* PaintBrushMat = Cast<UMaterialInterface>(StaticLoadObject(UMaterialInterface::StaticClass(), nullptr, TEXT("/LidarPointCloud/Materials/M_LidarSelectionBrush.M_LidarSelectionBrush")));
+
+	PaintBrush = NewObject<UStaticMeshComponent>();
+	PaintBrush->SetMobility(EComponentMobility::Movable);
+	PaintBrush->SetVisibility(false);
+	PaintBrush->SetStaticMesh(PaintBrushMesh);
+	PaintBrush->SetMaterial(0, PaintBrushMat);
+	PreviewScene.AddComponent(PaintBrush, FTransform::Identity);
+
 	ViewportOverlay->AddSlot()
 		.VAlign(VAlign_Top)
 		.HAlign(HAlign_Left)
-		.Padding(FMargin(10.0f, 40.0f, 10.0f, 10.0f))
+		.Padding(FMargin(4.0f, 35.0f))
 		[
-			SAssignNew(OverlayTextVerticalBox, SVerticalBox)
+			SAssignNew(OverlayTextBackground, SBorder)
+				.BorderBackgroundColor(FSlateColor(FLinearColor(0,0,0,0)))
+				.BorderImage(FCoreStyle::Get().GetBrush("GenericWhiteBox"))
+				.Padding(FMargin(10.0f))
+				[
+					SAssignNew(OverlayTextVerticalBox, SVerticalBox)
+				]
 		];
 }
 
@@ -105,6 +123,7 @@ SLidarPointCloudEditorViewport::~SLidarPointCloudEditorViewport()
 void SLidarPointCloudEditorViewport::AddReferencedObjects(FReferenceCollector& Collector)
 {
 	Collector.AddReferencedObject(PreviewCloudComponent);
+	Collector.AddReferencedObject(PaintBrush);
 	Collector.AddReferencedObject(PointCloud);
 }
 
@@ -145,6 +164,8 @@ void SLidarPointCloudEditorViewport::PopulateOverlayText(const TArray<FOverlayTe
 				.TextStyle(FEditorStyle::Get(), TextItem.Style)
 			];
 	}
+
+	OverlayTextBackground->SetBorderBackgroundColor(FSlateColor(FLinearColor(0, 0, 0, TextItems.Num() > 0 ? 0.45f : 0)));
 }
 
 bool SLidarPointCloudEditorViewport::IsVisible() const
diff --git a/Engine/Plugins/Enterprise/LidarPointCloud/Source/LidarPointCloudEditor/Private/LidarPointCloudEditorViewport.h b/Engine/Plugins/Enterprise/LidarPointCloud/Source/LidarPointCloudEditor/Private/LidarPointCloudEditorViewport.h
--- a/Engine/Plugins/Enterprise/LidarPointCloud/Source/LidarPointCloudEditor/Private/LidarPointCloudEditorViewport.h
+++ b/Engine/Plugins/Enterprise/LidarPointCloud/Source/LidarPointCloudEditor/Private/LidarPointCloudEditorViewport.h
@@ -15,6 +15,7 @@ class FLidarPointCloudEditorViewportClient;
 class SVerticalBox;
 class ULidarPointCloud;
 class ULidarPointCloudComponent;
+class UStaticMeshComponent;
 
 /**
  * PointCloud Editor Preview viewport widget
@@ -46,6 +47,8 @@ public:
 	/** Component for the preview point cloud. */
 	ULidarPointCloudComponent* PreviewCloudComponent;
 
+	UStaticMeshComponent* PaintBrush;
+
 	/**
 	 *	Sets up the point cloud that the Point Cloud editor is viewing.
 	 *
@@ -76,6 +79,8 @@ public:
 	virtual void OnFloatingButtonClicked() override {}
 	// End of ICommonEditorViewportToolbarInfoProvider interface
 
+	TSharedPtr<FLidarPointCloudEditorViewportClient> GetEditorViewportClient() { return EditorViewportClient; }
+
 protected:
 	/** SEditorViewport interface */
 	virtual TSharedRef<FEditorViewportClient> MakeEditorViewportClient() override;
@@ -109,4 +114,7 @@ private:
 
 	/** Pointer to the vertical box into which the overlay text items are added */
 	TSharedPtr<SVerticalBox> OverlayTextVerticalBox;
+
+	/** Pointer to the background of the overlay text */
+	TSharedPtr<SBorder> OverlayTextBackground;
 };
diff --git a/Engine/Plugins/Enterprise/LidarPointCloud/Source/LidarPointCloudEditor/Private/LidarPointCloudEditorViewportClient.cpp b/Engine/Plugins/Enterprise/LidarPointCloud/Source/LidarPointCloudEditor/Private/LidarPointCloudEditorViewportClient.cpp
--- a/Engine/Plugins/Enterprise/LidarPointCloud/Source/LidarPointCloudEditor/Private/LidarPointCloudEditorViewportClient.cpp
+++ b/Engine/Plugins/Enterprise/LidarPointCloud/Source/LidarPointCloudEditor/Private/LidarPointCloudEditorViewportClient.cpp
@@ -9,22 +9,116 @@
 
 #include "ConvexVolume.h"
 #include "CanvasTypes.h"
+#include "CanvasItem.h"
 #include "Settings/LevelEditorViewportSettings.h"
 #include "Editor/EditorPerProjectUserSettings.h"
 #include "Classes/EditorStyleSettings.h"
 #include "AssetViewerSettings.h"
 
+#include "GeomTools.h"
+
 #define LOCTEXT_NAMESPACE "FLidarPointCloudEditorViewportClient"
 
-namespace {
-	static const float LightRotSpeed = 0.22f;
-	static const float StaticMeshEditor_RotateSpeed = 0.01f;
-	static const float	StaticMeshEditor_TranslateSpeed = 0.25f;
-	static const float GridSize = 2048.0f;
-	static const int32 CellSize = 16;
-	static const float AutoViewportOrbitCameraTranslate = 256.0f;
+namespace UE {
+	namespace Lidar {
+		namespace Private {
+			namespace Editor
+			{
+				// Distance Square between the first and last points of the polygonal selection, where the shape will be considered as closed
+				constexpr int32 PolySnapDistanceSq = 40;
+
+				// Affects the frequency of new point injections when drawing lasso-based shapes
+				constexpr int32 LassoSpacingSq = 400;
+
+				// Affects the max depth delta when painting. Prevents the brush from "falling through" the gaps.
+				constexpr float PaintMaxDeviation = 0.15f;
+
+				// Defaults for the common draw helper
+				constexpr float GridSize = 2048.0f;
+				constexpr int32 CellSize = 16;
+			}
+		}
+	}
+}
+
+TArray<FVector2D> ToVectorArray(const TArray<FIntPoint>& Points)
+{
+	TArray<FVector2D> VectorPoints;
+	VectorPoints.Reserve(Points.Num());
+	for (const FIntPoint& Point : Points)
+	{
+		VectorPoints.Add(Point);
+	}
+	return VectorPoints;
+}
+
+// Slow, O(n2), but sufficient for the current problem
+bool IsPolygonSelfIntersecting(const TArray<FVector2D>& Points, bool bAllowLooping)
+{
+	const int32 MaxIndex = bAllowLooping ? Points.Num() : Points.Num() - 1;
+
+	for (int32 i = 0; i < MaxIndex; ++i)
+	{
+		const int32 i1 = i < Points.Num() - 1 ? i + 1 : 0;
+
+		const FVector2D P1 = Points[i];
+		const FVector2D P2 = Points[i1];
 
-	static float AmbientCubemapIntensity = 0.4f;
+		for (int32 j = 0; j < MaxIndex; ++j)
+		{
+			const int32 j1 = j < Points.Num() - 1 ? j + 1 : 0;
+
+			if (j1 != i && j != i && j != i1)
+			{
+				// Modified, inlined FMath::SegmentIntersection2D
+				const FVector2D SegmentStartA = P1;
+				const FVector2D SegmentEndA = P2;
+				const FVector2D SegmentStartB = Points[j];
+				const FVector2D SegmentEndB = Points[j1];
+				const FVector2D VectorA = P2 - SegmentStartA;
+				const FVector2D VectorB = SegmentEndB - SegmentStartB;
+
+				const float S = (-VectorA.Y * (SegmentStartA.X - SegmentStartB.X) + VectorA.X * (SegmentStartA.Y - SegmentStartB.Y)) / (-VectorB.X * VectorA.Y + VectorA.X * VectorB.Y);
+				const float T = (VectorB.X * (SegmentStartA.Y - SegmentStartB.Y) - VectorB.Y * (SegmentStartA.X - SegmentStartB.X)) / (-VectorB.X * VectorA.Y + VectorA.X * VectorB.Y);
+
+				if (S >= 0 && S <= 1 && T >= 0 && T <= 1)
+				{
+					return true;
+				}
+			}
+		}
+	}
+
+	return false;
+}
+
+
+// Copied from GeomTools.cpp and converted to work with FIntPoint
+bool IsPolygonConvex(const TArray<FIntPoint>& Points)
+{
+	const int32 PointCount = Points.Num();
+	int32 Sign = 0;
+	for (int32 PointIndex = 0; PointIndex < PointCount; ++PointIndex)
+	{
+		const FIntPoint& A = Points[PointIndex];
+		const FIntPoint& B = Points[(PointIndex + 1) % PointCount];
+		const FIntPoint& C = Points[(PointIndex + 2) % PointCount];
+		int32 Det = (B.X - A.X) * (C.Y - B.Y) - (B.Y - A.Y) * (C.X - B.X);
+		int32 DetSign = FMath::Sign(Det);
+		if (DetSign != 0)
+		{
+			if (Sign == 0)
+			{
+				Sign = DetSign;
+			}
+			else if (Sign != DetSign)
+			{
+				return false;
+			}
+		}
+	}
+
+	return true;
 }
 
 FLidarPointCloudEditorViewportClient::FLidarPointCloudEditorViewportClient(TWeakPtr<FLidarPointCloudEditor> InPointCloudEditor, const TSharedRef<SLidarPointCloudEditorViewport>& InPointCloudEditorViewport, FAdvancedPreviewScene* InPreviewScene, ULidarPointCloud* InPreviewPointCloud, ULidarPointCloudComponent* InPreviewPointCloudComponent)
@@ -32,7 +126,9 @@ FLidarPointCloudEditorViewportClient::FLidarPointCloudEditorViewportClient(TWeak
 	, PointCloudComponent(InPreviewPointCloudComponent)
 	, PointCloudEditorPtr(InPointCloudEditor)
 	, PointCloudEditorViewportPtr(InPointCloudEditorViewport)
+	, SelectionMethod(ELidarPointCloudSelectionMethod::Box)
 	, SelectionMode(ELidarPointCloudSelectionMode::None)
+	, PaintingRadius(500)
 {
 	// Setup defaults for the common draw helper.
 	DrawHelper.bDrawPivot = false;
@@ -42,8 +138,8 @@ FLidarPointCloudEditorViewportClient::FLidarPointCloudEditorViewportClient(TWeak
 	DrawHelper.GridColorAxis = FColor(160, 160, 160);
 	DrawHelper.GridColorMajor = FColor(144, 144, 144);
 	DrawHelper.GridColorMinor = FColor(128, 128, 128);
-	DrawHelper.PerspectiveGridSize = GridSize;
-	DrawHelper.NumCells = DrawHelper.PerspectiveGridSize / (CellSize * 2);
+	DrawHelper.PerspectiveGridSize = UE::Lidar::Private::Editor::GridSize;
+	DrawHelper.NumCells = DrawHelper.PerspectiveGridSize / (UE::Lidar::Private::Editor::CellSize * 2);
 
 	SetViewMode(VMI_Unlit);
 
@@ -69,6 +165,8 @@ FLidarPointCloudEditorViewportClient::~FLidarPointCloudEditorViewportClient()
 
 void FLidarPointCloudEditorViewportClient::Tick(float DeltaSeconds)
 {
+	using UE::Lidar::Private::Editor::PaintMaxDeviation;
+
 	FEditorViewportClient::Tick(DeltaSeconds);
 
 	// Tick the preview scene world.
@@ -76,117 +174,208 @@ void FLidarPointCloudEditorViewportClient::Tick(float DeltaSeconds)
 	{
 		PreviewScene->GetWorld()->Tick(LEVELTICK_All, DeltaSeconds);
 	}
-}
 
-void FLidarPointCloudEditorViewportClient::MouseMove(FViewport* InViewport, int32 x, int32 y)
-{
-	FEditorViewportClient::MouseMove(InViewport, x, y);
+	// Process line traces if in Paint mode
+	if (SelectionMethod == ELidarPointCloudSelectionMethod::Paint)
+	{
+		if (TSharedPtr<FLidarPointCloudEditor> Editor = PointCloudEditorPtr.Pin())
+		{
+			ULidarPointCloud* PC = Editor->GetPointCloudBeingEdited();
+
+			if (Editor->IsEditMode())
+			{
+				const bool bPainting = Viewport->KeyState(EKeys::LeftMouseButton) && bLineTraceHit;
+				const float TraceRadius = FMath::Max(PC->GetEstimatedPointSpacing(), 0.5f);
+
+				const FLidarPointCloudRay Ray = DeprojectCurrentMousePosition();
+				if (FLidarPointCloudPoint* Point = PC->LineTraceSingle(Ray, TraceRadius, true))
+				{
+					const float NewDistance = FVector::Dist(Point->Location, Ray.Origin);
+					const float Deviation = (NewDistance - LineTraceDistance) / LineTraceDistance;
+
+					// If painting, prevent large depth changes
+					// If not, query larger trace radius - if it passes the deviation test, it was a gap
+					if (Deviation > PaintMaxDeviation && (bPainting || (FVector::Dist(PC->LineTraceSingle(Ray, TraceRadius * 6, true)->Location, Ray.Origin) - LineTraceDistance) / LineTraceDistance <= PaintMaxDeviation))
+					{
+						LineTraceHitPoint = Ray.Origin + Ray.GetDirection() * LineTraceDistance;
+					}
+					else
+					{
+						LineTraceHitPoint = Point->Location;
+						LineTraceDistance = NewDistance;
+					}
+
+					bLineTraceHit = true;
+				}
+				else if(!bPainting)
+				{
+					bLineTraceHit = false;
+				}
+			}
+		}
+	}
 }
 
 bool FLidarPointCloudEditorViewportClient::InputKey(FViewport* InViewport, int32 ControllerId, FKey Key, EInputEvent Event, float AmountDepressed, bool Gamepad)
 {
-	bool bHandled = FEditorViewportClient::InputKey(InViewport, ControllerId, Key, Event, AmountDepressed, false);
-
-	// Handle viewport screenshot.
-	bHandled |= InputTakeScreenshot(InViewport, Key, Event);
+	using UE::Lidar::Private::Editor::PolySnapDistanceSq;
 
-	bHandled |= AdvancedPreviewScene->HandleInputKey(InViewport, ControllerId, Key, Event, AmountDepressed, Gamepad);
+	bool bHandled = false;
 
-	const bool bAltPressed = Viewport->KeyState(EKeys::LeftAlt) || Viewport->KeyState(EKeys::RightAlt);
-	const bool bCtrlPressed = Viewport->KeyState(EKeys::LeftControl) || Viewport->KeyState(EKeys::RightControl);
-	const bool bShiftPressed = Viewport->KeyState(EKeys::LeftShift) || Viewport->KeyState(EKeys::RightShift);
+	const bool bAlt = Key == EKeys::LeftAlt || Key == EKeys::RightAlt;
+	const bool bCtrl = Key == EKeys::LeftControl || Key == EKeys::RightControl;
 
-	// Edit Mode
-	if (PointCloudEditorPtr.Pin()->IsEditMode())
+	if (TSharedPtr<FLidarPointCloudEditor> Editor = PointCloudEditorPtr.Pin())
 	{
-		if (Event == IE_Pressed)
+		// Edit Mode
+		if (Editor->IsEditMode())
 		{
-			if (Key == EKeys::Delete)
+			if (Event == IE_Pressed)
 			{
-				if (bShiftPressed)
+				if (Key == EKeys::Delete)
 				{
-					PointCloudEditorPtr.Pin()->DeletePoints();
+					if (Viewport->KeyState(EKeys::LeftShift) || Viewport->KeyState(EKeys::RightShift))
+					{
+						Editor->DeletePoints();
+					}
+					else
+					{
+						Editor->HidePoints();
+					}
+
+					bHandled = true;
 				}
-				else
+				else if (Key == EKeys::Escape)
 				{
-					PointCloudEditorPtr.Pin()->HidePoints();
+					Editor->DeselectPoints();
+					SelectionPoints.Empty();
+					bHandled = true;
 				}
-
-				bHandled = true;
-			}
-			else if (Key == EKeys::Escape)
-			{
-				PointCloudEditorPtr.Pin()->DeselectPoints();
-				bHandled = true;
-			}
-		}
-
-		// Store current selection mode
-		ELidarPointCloudSelectionMode LastSelectionMode = SelectionMode;
-
-		// Update selection status
-		if (Key == EKeys::LeftMouseButton)
-		{
-			if (Viewport->KeyState(EKeys::LeftMouseButton))
-			{
-				if (bShiftPressed)
+				else if (Key == EKeys::Enter)
 				{
-					SelectionMode = ELidarPointCloudSelectionMode::Replace;
+					if (SelectionMethod == ELidarPointCloudSelectionMethod::Polygonal)
+					{
+						OnPolygonalSelectionEnd();
+						SelectionPoints.Empty();
+						bHandled = true;
+					}
 				}
-				else if (bCtrlPressed)
+				else if (bAlt)
 				{
+					bHandled = true;
+					SelectionMode = ELidarPointCloudSelectionMode::Subtract;
+				}
+				else if (bCtrl)
+				{
+					bHandled = true;
 					SelectionMode = ELidarPointCloudSelectionMode::Add;
 				}
-				else if (bAltPressed)
+				else if (Key == EKeys::LeftMouseButton)
 				{
-					SelectionMode = ELidarPointCloudSelectionMode::Subtract;
+					bHandled = true;
+
+					// Selection start
+					if (SelectionMethod == ELidarPointCloudSelectionMethod::Box)
+					{
+						// Mark the cursor location for selection start
+						InViewport->GetMousePos(SelectionPoints[SelectionPoints.AddUninitialized()]);
+					}
+					else if (SelectionMethod == ELidarPointCloudSelectionMethod::Polygonal)
+					{
+						// Add new lasso point
+						FIntPoint NewPoint;
+						InViewport->GetMousePos(NewPoint);
+
+						// Don't allow duplicates
+						if (SelectionPoints.Num() == 0 || SelectionPoints.Last() != NewPoint)
+						{
+							TArray<FVector2D> VectorPoints = ToVectorArray(SelectionPoints);
+							VectorPoints.Add(NewPoint);
+
+							if (!IsPolygonSelfIntersecting(VectorPoints, false))
+							{
+								// Snap to first point
+								if (SelectionPoints.Num() > 1 && (NewPoint - SelectionPoints[0]).SizeSquared() < PolySnapDistanceSq)
+								{
+									OnPolygonalSelectionEnd();
+									SelectionPoints.Empty();
+								}
+								else
+								{
+									SelectionPoints.Add(NewPoint);
+								}
+							}
+						}
+					}
+					else if (SelectionMethod == ELidarPointCloudSelectionMethod::Lasso)
+					{
+						InViewport->GetMousePos(SelectionPoints[SelectionPoints.AddUninitialized()]);
+					}
+					else if (SelectionMethod == ELidarPointCloudSelectionMethod::Paint)
+					{
+						OnPaintSelection();
+					}
 				}
-				else
+
+				if (SelectionMethod == ELidarPointCloudSelectionMethod::Paint)
 				{
-					SelectionMode = ELidarPointCloudSelectionMode::None;
+					// Do not block ability to change camera speed
+					if (!Viewport->KeyState(EKeys::RightMouseButton))
+					{
+						if (Key == EKeys::MouseScrollUp)
+						{
+							PaintingRadius *= 1.1f;
+							bHandled = true;
+						}
+						else if (Key == EKeys::MouseScrollDown)
+						{
+							PaintingRadius /= 1.1f;
+							bHandled = true;
+						}
+					}
 				}
 			}
-			else
+			else if (Event == IE_Released)
 			{
-				SelectionMode = ELidarPointCloudSelectionMode::None;
+				if (bAlt || bCtrl)
+				{
+					SelectionMode = ELidarPointCloudSelectionMode::None;
+					bHandled = true;
+				}
+				else if (Key == EKeys::LeftMouseButton)
+				{
+					bHandled = true;
+
+					if (SelectionMethod == ELidarPointCloudSelectionMethod::Box)
+					{
+						// Mark the cursor location for selection end
+						InViewport->GetMousePos(SelectionPoints[SelectionPoints.AddUninitialized()]);
+						OnBoxSelectionEnd();
+						SelectionPoints.Empty();
+					}
+					else if (SelectionMethod == ELidarPointCloudSelectionMethod::Lasso)
+					{
+						OnLassoSelectionEnd();
+						SelectionPoints.Empty();
+					}
+				}
 			}
-
-			bHandled = true;
 		}
-
-		// Check if selection mode has changed since the last frame
-		if (SelectionMode != LastSelectionMode)
+		// Navigation Mode
+		else
 		{
-			// Selection start
-			if (SelectionMode != ELidarPointCloudSelectionMode::None)
-			{
-				// Mark the cursor location for selection start
-				InViewport->GetMousePos(SelectionStartLocation);
-
-				OnSelectionStart();
-			}
-			// Selection end
-			else
-			{
-				FIntPoint SelectionEndLocation;
 
-				// Mark the cursor location for selection end
-				InViewport->GetMousePos(SelectionEndLocation);
-
-				FIntVector4 SelectionArea;
-
-				SelectionArea.X = FMath::Min(SelectionStartLocation.X, SelectionEndLocation.X);
-				SelectionArea.Y = FMath::Min(SelectionStartLocation.Y, SelectionEndLocation.Y);
-				SelectionArea.Z = FMath::Max(SelectionStartLocation.X, SelectionEndLocation.X);
-				SelectionArea.W = FMath::Max(SelectionStartLocation.Y, SelectionEndLocation.Y);
-
-				OnSelectionEnd(SelectionArea, LastSelectionMode);
-			}
 		}
 	}
-	// Navigation Mode
-	else
+
+	if (!bHandled)
 	{
+		bHandled = FEditorViewportClient::InputKey(InViewport, ControllerId, Key, Event, AmountDepressed, false);
+
+		// Handle viewport screenshot.
+		bHandled |= InputTakeScreenshot(InViewport, Key, Event);
+
+		bHandled |= AdvancedPreviewScene->HandleInputKey(InViewport, ControllerId, Key, Event, AmountDepressed, Gamepad);
 
 	}
 
@@ -195,96 +384,175 @@ bool FLidarPointCloudEditorViewportClient::InputKey(FViewport* InViewport, int32
 
 bool FLidarPointCloudEditorViewportClient::InputAxis(FViewport* InViewport, int32 ControllerId, FKey Key, float Delta, float DeltaTime, int32 NumSamples, bool bGamepad)
 {
-	bool bResult = true;
+	using UE::Lidar::Private::Editor::LassoSpacingSq;
+
+	bool bHandled = false;
 
 	if (!bDisableInput)
 	{
-		if (SelectionMode != ELidarPointCloudSelectionMode::None)
+		if (TSharedPtr<FLidarPointCloudEditor> Editor = PointCloudEditorPtr.Pin())
 		{
-			// Do not move the camera while in selection
+			if (Editor->IsEditMode())
+			{
+				if (Viewport->KeyState(EKeys::LeftMouseButton))
+				{
+					if (SelectionMethod == ELidarPointCloudSelectionMethod::Lasso)
+					{
+						FIntPoint NewPoint;
+						InViewport->GetMousePos(NewPoint);
+
+						// Check if the spacing is sufficient
+						if ((NewPoint - SelectionPoints.Last()).SizeSquared() > LassoSpacingSq)
+						{
+							SelectionPoints.Add(NewPoint);
+						}
+
+						bHandled = true;
+					}
+					else if (SelectionMethod == ELidarPointCloudSelectionMethod::Paint)
+					{
+						OnPaintSelection();
+						bHandled = true;
+					}
+				}
+			}
 		}
-		else
+
+		if(!bHandled)
 		{
-			bResult = AdvancedPreviewScene->HandleViewportInput(InViewport, ControllerId, Key, Delta, DeltaTime, NumSamples, bGamepad);
-			if (bResult)
+			bHandled = AdvancedPreviewScene->HandleViewportInput(InViewport, ControllerId, Key, Delta, DeltaTime, NumSamples, bGamepad);
+			if (bHandled)
 			{
 				Invalidate();
 			}
 			else
 			{
-				bResult = FEditorViewportClient::InputAxis(InViewport, ControllerId, Key, Delta, DeltaTime, NumSamples, bGamepad);
+				bHandled = FEditorViewportClient::InputAxis(InViewport, ControllerId, Key, Delta, DeltaTime, NumSamples, bGamepad);
 			}
 		}
 	}
 
-	return bResult;
-}
-
-void FLidarPointCloudEditorViewportClient::ProcessClick(FSceneView& InView, HHitProxy* HitProxy, FKey Key, EInputEvent Event, uint32 HitX, uint32 HitY) { }
-
-void FLidarPointCloudEditorViewportClient::Draw(const FSceneView* View, FPrimitiveDrawInterface* PDI)
-{
-	FEditorViewportClient::Draw(View, PDI);
+	return bHandled;
 }
 
 void FLidarPointCloudEditorViewportClient::DrawCanvas(FViewport& InViewport, FSceneView& View, FCanvas& Canvas)
 {
-	TSharedPtr<FLidarPointCloudEditor> PointCloudEditor = PointCloudEditorPtr.Pin();
-	TSharedPtr<SLidarPointCloudEditorViewport> PointCloudEditorViewport = PointCloudEditorViewportPtr.Pin();
-	if (!PointCloudEditor.IsValid() || !PointCloudEditorViewport.IsValid())
-	{
-		return;
-	}
-
 	TArray<SLidarPointCloudEditorViewport::FOverlayTextItem> TextItems;
 
-	if (PointCloudEditorPtr.Pin()->IsEditMode())
+	if (TSharedPtr<FLidarPointCloudEditor> Editor = PointCloudEditorPtr.Pin())
 	{
-		// Text entries
-		{
-			TextItems.Add(SLidarPointCloudEditorViewport::FOverlayTextItem(FText::FromString("[ESCAPE] to de-select")));
-			TextItems.Add(SLidarPointCloudEditorViewport::FOverlayTextItem(FText::FromString("[SHIFT] + Mouse Drag to set selection")));
-			TextItems.Add(SLidarPointCloudEditorViewport::FOverlayTextItem(FText::FromString("[CTRL] + Mouse Drag to add selection")));
-			TextItems.Add(SLidarPointCloudEditorViewport::FOverlayTextItem(FText::FromString("[ALT] + Mouse Drag to subtract selection")));
-			TextItems.Add(SLidarPointCloudEditorViewport::FOverlayTextItem(FText::FromString("")));
-			TextItems.Add(SLidarPointCloudEditorViewport::FOverlayTextItem(FText::FromString("Press [DELETE] to hide selected")));
-			TextItems.Add(SLidarPointCloudEditorViewport::FOverlayTextItem(FText::FromString("Press [SHIFT] + [DELETE] to delete points")));
-			TextItems.Add(SLidarPointCloudEditorViewport::FOverlayTextItem(FText::FromString("")));
-			TextItems.Add(SLidarPointCloudEditorViewport::FOverlayTextItem(FText::Format(FText::FromString("Selected Points:  {0}"), FText::AsNumber(PointCloudEditorPtr.Pin()->GetSelectedPoints().Num()))));
-		}
-
-		// Selection box
-		if (SelectionMode != ELidarPointCloudSelectionMode::None)
+		if (Editor->IsEditMode())
 		{
-			FIntPoint SelectionCurrentLocation;
-			Viewport->GetMousePos(SelectionCurrentLocation);
+			TArray<FString> Labels;
 
-			const float InvScale = 1.0f / Viewport->GetClient()->GetDPIScale();
-
-			const float X = FMath::Min(SelectionStartLocation.X, SelectionCurrentLocation.X) * InvScale;
-			const float Y = FMath::Min(SelectionStartLocation.Y, SelectionCurrentLocation.Y) * InvScale;
-			const float SizeX = FMath::Max(SelectionStartLocation.X, SelectionCurrentLocation.X) * InvScale - X;
-			const float SizeY = FMath::Max(SelectionStartLocation.Y, SelectionCurrentLocation.Y) * InvScale - Y;
-
-			FLinearColor SelectionColor = GetDefault<UEditorStyleSettings>()->SelectionColor;
-			SelectionColor.A = 0.35f;
+			// Add mode-specific labels
+			switch (SelectionMethod)
+			{
+			case ELidarPointCloudSelectionMethod::Box:
+				DrawSelectionBox(Canvas);
+				Labels.Append({
+					"BOX SELECTION MODE",
+					"",
+					"Click + Drag to replace selection",
+					"[CTRL] to add selection",
+					"[ALT] to subtract selection"
+				});
+				break;
+
+			case ELidarPointCloudSelectionMethod::Polygonal:
+				DrawSelectionPolygonal(Canvas);
+				Labels.Append({
+					"POLYGONAL SELECTION MODE",
+					"",
+					"[CTRL] to add selection",
+					"[ALT] to subtract selection"
+					});
+				break;
+
+			case ELidarPointCloudSelectionMethod::Lasso:
+				DrawSelectionLasso(Canvas);
+				Labels.Append({
+					"LASSO SELECTION MODE",
+					"",
+					"Click + Drag to paint selection shape",
+					"[CTRL] to add selection",
+					"[ALT] to subtract selection"
+				});
+				break;
+
+			case ELidarPointCloudSelectionMethod::Paint:
+				DrawSelectionPaint(Canvas);
+				Labels.Append({
+					"PAINT SELECTION MODE",
+					"",
+					"Click + Drag to paint selection",
+					"[SCROLL] to change brush size",
+					"[ALT] to subtract selection"
+				});
+				break;
+
+			default:
+				break;
+			}
 
-			Canvas.DrawTile(X, Y, SizeX, SizeY, 0, 0, 0, 0, SelectionColor);
+			// Add common labels
+			Labels.Append({
+				"",
+				"[ESCAPE] to de-select all points",
+				"[DELETE] to hide selected points",
+				"[SHIFT] + [DELETE] to permanently delete selected points",
+				"",
+				FString::Printf(TEXT("Selected Points: %d"), Editor->GetSelectedPoints().Num())
+			});
+
+			// Convert to text entries
+			TextItems.Reserve(Labels.Num());
+			for (const FString& Label : Labels)
+			{
+				TextItems.Add(SLidarPointCloudEditorViewport::FOverlayTextItem(FText::FromString(Label)));
+			}
 		}
 	}
 
-	PointCloudEditorViewport->PopulateOverlayText(TextItems);
+	if (TSharedPtr<SLidarPointCloudEditorViewport> PointCloudEditorViewport = PointCloudEditorViewportPtr.Pin())
+	{
+		PointCloudEditorViewport->PopulateOverlayText(TextItems);
+	}
 }
 
 bool FLidarPointCloudEditorViewportClient::ShouldOrbitCamera() const
 {
-	if (GetDefault<ULevelEditorViewportSettings>()->bUseUE3OrbitControls)
+	if (TSharedPtr<FLidarPointCloudEditor> Editor = PointCloudEditorPtr.Pin())
 	{
-		// this editor orbits always if orbit controls are enabled
-		return true;
+		if (Editor->IsEditMode())
+		{
+			return false;
+		}
 	}
 
-	return FEditorViewportClient::ShouldOrbitCamera();
+	return GetDefault<ULevelEditorViewportSettings>()->bUseUE3OrbitControls || FEditorViewportClient::ShouldOrbitCamera();
+}
+
+void FLidarPointCloudEditorViewportClient::LostFocus(FViewport* InViewport)
+{
+	FEditorViewportClient::LostFocus(InViewport);
+
+	// Cancel selection
+	SelectionMode = ELidarPointCloudSelectionMode::None;
+	SelectionPoints.Empty();
+}
+
+void FLidarPointCloudEditorViewportClient::ReceivedFocus(FViewport* InViewport)
+{
+	// This is needed if the user presses Alt / Ctrl / Shift before the client acquires focus
+	if (Viewport->KeyState(EKeys::LeftAlt) || Viewport->KeyState(EKeys::RightAlt))
+	{
+		SelectionMode = ELidarPointCloudSelectionMode::Subtract;
+	}
+	else if (Viewport->KeyState(EKeys::LeftControl) || Viewport->KeyState(EKeys::RightControl))
+	{
+		SelectionMode = ELidarPointCloudSelectionMode::Add;
+	}
 }
 
 void FLidarPointCloudEditorViewportClient::PerspectiveCameraMoved()
@@ -327,125 +595,392 @@ void FLidarPointCloudEditorViewportClient::SetAdvancedShowFlagsForScene(const bo
 	}
 }
 
-void FLidarPointCloudEditorViewportClient::OnSelectionStart()
+FSceneView* FLidarPointCloudEditorViewportClient::GetView()
 {
+	if (TSharedPtr<FLidarPointCloudEditor> Editor = PointCloudEditorPtr.Pin())
+	{
+		// Compute a view.
+		FSceneViewFamilyContext ViewFamily(FSceneViewFamily::ConstructionValues(Viewport, GetScene(), EngineShowFlags).SetRealtimeUpdate(IsRealtime()));
+		FSceneView* View = CalcSceneView(&ViewFamily);
+
+		const FVector LocationOffset = Editor->GetPointCloudBeingEdited()->LocationOffset.ToVector();
+
+		// Adjust for the LocationOffset
+		if (View->IsPerspectiveProjection())
+		{
+			View->ViewLocation -= LocationOffset;
+			View->UpdateViewMatrix();
+		}
+
+		return View;
+	}
+
+	return nullptr;
 }
 
-void FLidarPointCloudEditorViewportClient::OnSelectionEnd(const FIntVector4& SelectionArea, ELidarPointCloudSelectionMode Mode)
+FLidarPointCloudRay FLidarPointCloudEditorViewportClient::DeprojectCurrentMousePosition()
+{
+	FSceneView* View = GetView();
+	const FMatrix InvViewProjectionMatrix = View->ViewMatrices.GetInvViewProjectionMatrix();
+
+	FIntPoint CurrentMousePosition;
+	Viewport->GetMousePos(CurrentMousePosition);
+
+	FVector Origin, Direction;
+	FSceneView::DeprojectScreenToWorld(FVector2D(CurrentMousePosition), FIntRect(FIntPoint(0, 0), Viewport->GetSizeXY()), InvViewProjectionMatrix, Origin, Direction);
+
+	return FLidarPointCloudRay(Origin, Direction);
+}
+
+void FLidarPointCloudEditorViewportClient::OnBoxSelectionEnd()
 {
 	if (ULidarPointCloudComponent* PointCloudComponentRawPtr = PointCloudComponent.Get())
 	{
 		if (TSharedPtr<FLidarPointCloudEditor> Editor = PointCloudEditorPtr.Pin())
 		{
-			// Compute a view.
-			FSceneViewFamilyContext ViewFamily(FSceneViewFamily::ConstructionValues(Viewport, GetScene(), EngineShowFlags).SetRealtimeUpdate(IsRealtime()));
-			FSceneView* View = CalcSceneView(&ViewFamily);	
+			if (SelectionPoints[0] == SelectionPoints.Last())
+			{
+				Editor->DeselectPoints();
+				return;
+			}
 
-			const FVector LocationOffset = Editor->GetPointCloudBeingEdited()->LocationOffset.ToVector();
+			FIntVector4 SelectionArea;
+			SelectionArea.X = FMath::Min(SelectionPoints[0].X, SelectionPoints[1].X);
+			SelectionArea.Y = FMath::Min(SelectionPoints[0].Y, SelectionPoints[1].Y);
+			SelectionArea.Z = FMath::Max(SelectionPoints[0].X, SelectionPoints[1].X);
+			SelectionArea.W = FMath::Max(SelectionPoints[0].Y, SelectionPoints[1].Y);
 
-			// Adjust for the LocationOffset
-			if (View->IsPerspectiveProjection())
+			const FConvexVolume ConvexVolume = BuildConvexVolumeForPoints(TArray<FVector2D>({
+				FVector2D(SelectionArea.X, SelectionArea.Y),
+				FVector2D(SelectionArea.X, SelectionArea.W),
+				FVector2D(SelectionArea.Z, SelectionArea.W),
+				FVector2D(SelectionArea.Z, SelectionArea.Y) }));
+
+			if (SelectionMode == ELidarPointCloudSelectionMode::Subtract)
+			{
+				Editor->DeselectPointsByConvexVolume(ConvexVolume);
+			}
+			else
 			{
-				View->ViewLocation -= LocationOffset;
-				View->UpdateViewMatrix();
+				Editor->SelectPointsByConvexVolume(ConvexVolume, SelectionMode == ELidarPointCloudSelectionMode::Add);
 			}
-			
-			const FMatrix InvViewProjectionMatrix = View->ViewMatrices.GetInvViewProjectionMatrix();
-
-			/**
-			 * Compute 4 rays from the corners of the selection area
-			 *
-			 * 0 -- 3
-			 * |    |
-			 * 1 -- 2
-			 */
-			FVector Origins[6];
-			FVector Normals[6];
-			FVector Directions[4];
-			FSceneView::DeprojectScreenToWorld(FVector2D(SelectionArea.X, SelectionArea.Y), FIntRect(FIntPoint(0, 0), Viewport->GetSizeXY()), InvViewProjectionMatrix, Origins[0], Directions[0]);
-			FSceneView::DeprojectScreenToWorld(FVector2D(SelectionArea.X, SelectionArea.W), FIntRect(FIntPoint(0, 0), Viewport->GetSizeXY()), InvViewProjectionMatrix, Origins[1], Directions[1]);
-			FSceneView::DeprojectScreenToWorld(FVector2D(SelectionArea.Z, SelectionArea.W), FIntRect(FIntPoint(0, 0), Viewport->GetSizeXY()), InvViewProjectionMatrix, Origins[2], Directions[2]);
-			FSceneView::DeprojectScreenToWorld(FVector2D(SelectionArea.Z, SelectionArea.Y), FIntRect(FIntPoint(0, 0), Viewport->GetSizeXY()), InvViewProjectionMatrix, Origins[3], Directions[3]);
-
-			const FVector& ViewDirection = View->GetViewDirection();
-
-			// Shared calculations
-			Normals[4] = ViewDirection;
-			Normals[5] = -ViewDirection;
-			Origins[4] = Origins[0] + ViewDirection * 99999999.0f;
-
-			// Perspective View
-			if (View->IsPerspectiveProjection())
+		}
+	}
+}
+
+void FLidarPointCloudEditorViewportClient::OnPolygonalSelectionEnd()
+{
+	// Skip invalid selections
+	if (SelectionPoints.Num() < 3)
+	{
+		return;
+	}
+
+	if (ULidarPointCloudComponent* PointCloudComponentRawPtr = PointCloudComponent.Get())
+	{
+		if (TSharedPtr<FLidarPointCloudEditor> Editor = PointCloudEditorPtr.Pin())
+		{
+			TArray<TArray<FVector2D>> ConvexShapes;
+			TArray<FVector2D> VectorPoints = ToVectorArray(SelectionPoints);
+
+			if (IsPolygonConvex(SelectionPoints))
 			{
-				Origins[5] = Origins[0];
-
-				// Calculate plane normals
-				Normals[0] = Directions[1] ^ Directions[0];
-				Normals[1] = Directions[2] ^ Directions[1];
-				Normals[2] = Directions[3] ^ Directions[2];
-				Normals[3] = Directions[0] ^ Directions[3];
-
-				Normals[0].Normalize();
-				Normals[1].Normalize();
-				Normals[2].Normalize();
-				Normals[3].Normalize();
+				ConvexShapes.Add(VectorPoints);
 			}
-			// Ortho Views
 			else
 			{
-				// Adjust for the LocationOffset
-				for (int32 i = 0; i < 5; ++i)
+				// Check for self-intersecting shape
+				if (!IsPolygonSelfIntersecting(VectorPoints, true))
 				{
-					Origins[i] -= LocationOffset;
+					// The separation needs points in CCW order
+					if (!FGeomTools2D::IsPolygonWindingCCW(VectorPoints))
+					{
+						Algo::Reverse(VectorPoints);
+					}
+
+					TArray<FVector2D> Triangles;
+					FGeomTools2D::TriangulatePoly(Triangles, VectorPoints, false);
+					FGeomTools2D::GenerateConvexPolygonsFromTriangles(ConvexShapes, Triangles);
 				}
+			}
 
-				Origins[5] = -Origins[4];
+			for (int32 i = 0; i < ConvexShapes.Num(); ++i)
+			{
+				const FConvexVolume ConvexVolume = BuildConvexVolumeForPoints(ConvexShapes[i]);
 
-				// Left / Right view
-				if (ViewDirection.X != 0)
+				if (SelectionMode == ELidarPointCloudSelectionMode::Subtract)
 				{
-					Normals[0] = -ViewDirection.X * FVector::RightVector;
-					Normals[1] = -ViewDirection.X * FVector::UpVector;
-					Normals[2] = ViewDirection.X * FVector::RightVector;
-					Normals[3] = ViewDirection.X * FVector::UpVector;
+					Editor->DeselectPointsByConvexVolume(ConvexVolume);
 				}
-				// Front / Back view
-				else if (ViewDirection.Y != 0)
+				else
 				{
-					Normals[0] = ViewDirection.Y * FVector::ForwardVector;
-					Normals[1] = ViewDirection.Y * FVector::UpVector;
-					Normals[2] = -ViewDirection.Y * FVector::ForwardVector;
-					Normals[3] = -ViewDirection.Y * FVector::UpVector;
+					// Consecutive shapes need to be additive
+					Editor->SelectPointsByConvexVolume(ConvexVolume, i > 0 || SelectionMode == ELidarPointCloudSelectionMode::Add);
 				}
-				// Top / Bottom view
-				else
+			}
+		}
+	}
+}
+
+void FLidarPointCloudEditorViewportClient::OnLassoSelectionEnd()
+{
+	OnPolygonalSelectionEnd();
+}
+
+void FLidarPointCloudEditorViewportClient::OnPaintSelection()
+{
+	if (TSharedPtr<FLidarPointCloudEditor> Editor = PointCloudEditorPtr.Pin())
+	{
+		if (bLineTraceHit)
+		{
+			if (SelectionMode == ELidarPointCloudSelectionMode::Subtract)
+			{
+				Editor->DeselectPointsBySphere(FSphere(LineTraceHitPoint, PaintingRadius));
+			}
+			else
+			{
+				Editor->SelectPointsBySphere(FSphere(LineTraceHitPoint, PaintingRadius));
+			}
+		}
+	}
+}
+
+void FLidarPointCloudEditorViewportClient::DrawSelectionBox(FCanvas& Canvas)
+{
+	if (SelectionPoints.Num() == 0)
+	{
+		return;
+	}
+
+	const FIntPoint SelectionStartLocation = SelectionPoints[0];
+	FIntPoint SelectionCurrentLocation;
+	Viewport->GetMousePos(SelectionCurrentLocation);
+
+	const float InvScale = 1.0f / Viewport->GetClient()->GetDPIScale();
+
+	const float X = FMath::Min(SelectionStartLocation.X, SelectionCurrentLocation.X) * InvScale;
+	const float Y = FMath::Min(SelectionStartLocation.Y, SelectionCurrentLocation.Y) * InvScale;
+	const float SizeX = FMath::Max(SelectionStartLocation.X, SelectionCurrentLocation.X) * InvScale - X;
+	const float SizeY = FMath::Max(SelectionStartLocation.Y, SelectionCurrentLocation.Y) * InvScale - Y;
+
+	FLinearColor SelectionColor = GetDefault<UEditorStyleSettings>()->SelectionColor;
+	SelectionColor.A = 0.35f;
+
+	Canvas.DrawTile(X, Y, SizeX, SizeY, 0, 0, 0, 0, SelectionColor);
+
+	// Selection Border
+	{
+		FCanvasLineItem Line;
+		Line.SetColor(GetDefault<UEditorStyleSettings>()->SelectionColor);
+		Line.LineThickness = 2;
+
+		Line.Origin = FVector(SelectionStartLocation);
+		Line.EndPos = FVector(SelectionCurrentLocation.X, SelectionStartLocation.Y, 0);
+		Canvas.DrawItem(Line);
+
+		Line.Origin = Line.EndPos;
+		Line.EndPos = FVector(SelectionCurrentLocation);
+		Canvas.DrawItem(Line);
+
+		Line.Origin = Line.EndPos;
+		Line.EndPos = FVector(SelectionStartLocation.X, SelectionCurrentLocation.Y, 0);
+		Canvas.DrawItem(Line);
+
+		Line.Origin = Line.EndPos;
+		Line.EndPos = FVector(SelectionStartLocation);
+		Canvas.DrawItem(Line);
+	}
+}
+
+void FLidarPointCloudEditorViewportClient::DrawSelectionPolygonal(FCanvas& Canvas)
+{
+	using UE::Lidar::Private::Editor::PolySnapDistanceSq;
+
+	if (SelectionPoints.Num() == 0)
+	{
+		return;
+	}
+
+	// Create a local copy of selection points, injecting the current mouse position at the end
+	TArray<FIntPoint> DrawSelectionPoints = SelectionPoints;
+	Viewport->GetMousePos(DrawSelectionPoints[DrawSelectionPoints.AddUninitialized()]);
+
+	if (DrawSelectionPoints.Last() == DrawSelectionPoints.Last(1))
+	{
+		DrawSelectionPoints.RemoveAt(DrawSelectionPoints.Num() - 1, 1, false);
+	}
+
+	TArray<FVector2D> VectorPoints = ToVectorArray(DrawSelectionPoints);
+
+	// Calculate visual indication of complete polygon for the user
+	const bool bPolyComplete = DrawSelectionPoints.Num() > 2 && (DrawSelectionPoints.Last() - DrawSelectionPoints[0]).SizeSquared() < PolySnapDistanceSq;
+	const bool bSelfIntersecting = DrawSelectionPoints.Num() > 2 && IsPolygonSelfIntersecting(VectorPoints, true);
+	FLinearColor SelectionColor = bSelfIntersecting ? FLinearColor::Red : bPolyComplete ? FLinearColor::Green : GetDefault<UEditorStyleSettings>()->SelectionColor;
+
+	// Selection Area
+	if (VectorPoints.Num() > 2 && !bSelfIntersecting)
+	{
+		FCanvasUVTri Tri;
+		Tri.V0_Pos = Tri.V1_Pos = Tri.V2_Pos = Tri.V0_UV = Tri.V1_UV = Tri.V2_UV = FVector2D::ZeroVector;
+		SelectionColor.A = 0.35f;
+		Tri.V0_Color = Tri.V1_Color = Tri.V2_Color = SelectionColor;
+
+		TArray<FCanvasUVTri> TriangleList;
+		TArray<TArray<FVector2D>> Polygons;
+
+		if (IsPolygonConvex(DrawSelectionPoints))
+		{
+			Polygons.Add(VectorPoints);
+			TriangleList.Reserve(VectorPoints.Num() - 2);
+		}
+		else
+		{
+			// The separation needs points in CCW order
+			if (!FGeomTools2D::IsPolygonWindingCCW(VectorPoints))
+			{
+				Algo::Reverse(VectorPoints);
+			}
+
+			TArray<FVector2D> Triangles;
+			if (FGeomTools2D::TriangulatePoly(Triangles, VectorPoints, false))
+			{
+				FGeomTools2D::GenerateConvexPolygonsFromTriangles(Polygons, Triangles);
+
+				// Calculate the number of triangles and reserve space
+				int32 NumTriangles = 0;
+				for (const TArray<FVector2D>& Polygon : Polygons)
 				{
-					Normals[0] = ViewDirection.Z * FVector::ForwardVector;
-					Normals[1] = -ViewDirection.Z * FVector::RightVector;
-					Normals[2] = -ViewDirection.Z * FVector::ForwardVector;
-					Normals[3] = ViewDirection.Z * FVector::RightVector;
+					NumTriangles += Polygon.Num() - 2;
 				}
+				TriangleList.Reserve(NumTriangles);
 			}
+		}
+
+		for (const TArray<FVector2D>& Polygon : Polygons)
+		{
+			for (int32 i = 2; i < Polygon.Num(); ++i)
+			{
+				Tri.V0_Pos = Polygon[i];
+				Tri.V1_Pos = Polygon[0];
+				Tri.V2_Pos = Polygon[i - 1];
+				TriangleList.Add(Tri);
+			}
+		}
+
+		if (TriangleList.Num() > 0)
+		{
+			FCanvasTriangleItem Selection(TriangleList, GWhiteTexture);
+			Selection.BlendMode = SE_BLEND_AlphaBlend;
+			Canvas.DrawItem(Selection);
+		}
+	}
+
+	// Selection Border
+	{
+		FCanvasLineItem Line;
+		Line.SetColor(SelectionColor);
+		Line.LineThickness = 2;
+
+		for (int32 i = 1; i < VectorPoints.Num(); ++i)
+		{
+			Line.Origin = FVector(VectorPoints[i], 0);
+			Line.EndPos = FVector(VectorPoints[i - 1], 0);
+			Canvas.DrawItem(Line);
+		}
+	}
+}
+
+void FLidarPointCloudEditorViewportClient::DrawSelectionLasso(FCanvas& Canvas)
+{
+	DrawSelectionPolygonal(Canvas);
+}
+
+void FLidarPointCloudEditorViewportClient::DrawSelectionPaint(FCanvas& Canvas)
+{
+	if (TSharedPtr<FLidarPointCloudEditor> Editor = PointCloudEditorPtr.Pin())
+	{
+		if (TSharedPtr<SLidarPointCloudEditorViewport> EditorViewport = PointCloudEditorViewportPtr.Pin())
+		{
+			EditorViewport->PaintBrush->SetVisibility(bLineTraceHit);
+			EditorViewport->PaintBrush->SetWorldScale3D(FVector(PaintingRadius * 0.02f));
+			EditorViewport->PaintBrush->SetWorldLocation(LineTraceHitPoint);
+		}
+	}
+}
 
-			// Compute selection frustum planes from the rays
-			TArray<FPlane, TInlineAllocator<6>> Planes;
+FConvexVolume FLidarPointCloudEditorViewportClient::BuildConvexVolumeForPoints(const TArray<FVector2D>& Points)
+{
+	FConvexVolume ConvexVolume;
+
+	if (TSharedPtr<FLidarPointCloudEditor> Editor = PointCloudEditorPtr.Pin())
+	{
+		FSceneView* View = GetView();
+		const FVector LocationOffset = Editor->GetPointCloudBeingEdited()->LocationOffset.ToVector();
+		const FMatrix InvViewProjectionMatrix = View->ViewMatrices.GetInvViewProjectionMatrix();
 
-			// Left, Bottom, Right, Top, Forward, Backward
-			for (int32 i = 0; i < 6; ++i)
+		TArray<FVector> Origins; Origins.AddUninitialized(Points.Num() + 2);
+		TArray<FVector> Normals; Normals.AddUninitialized(Points.Num() + 2);
+		TArray<FVector> Directions; Directions.AddUninitialized(Points.Num());
+		FVector MeanCenter = FVector::ZeroVector;
+
+		for (int32 i = 0; i < Points.Num(); ++i)
+		{
+			FSceneView::DeprojectScreenToWorld(Points[i], FIntRect(FIntPoint(0, 0), Viewport->GetSizeXY()), InvViewProjectionMatrix, Origins[i], Directions[i]);
+			MeanCenter += Origins[i];
+		}
+
+		MeanCenter /= Points.Num();
+
+		const FVector& ViewDirection = View->GetViewDirection();
+
+		// Shared calculations
+		Normals.Last(1) = ViewDirection;
+		Normals.Last() = -ViewDirection;
+		Origins.Last(1) = Origins[0] + ViewDirection * 99999999.0f;
+
+		// Calculate plane normals
+		bool bFlipNormals = false;
+		for (int32 i = 0; i < Points.Num(); ++i)
+		{
+			Normals[i] = ((Origins[(i + 1) % Points.Num()] - Origins[i]).GetSafeNormal() ^ Directions[i]).GetSafeNormal();
+
+			if (i == 0)
 			{
-				Planes.Emplace(Origins[i], Normals[i]);
+				bFlipNormals = FVector::DotProduct(Normals[i], (MeanCenter - Origins[i])) > 0;
 			}
 
-			if (Mode == ELidarPointCloudSelectionMode::Subtract)
+			if (bFlipNormals)
 			{
-				Editor->DeselectPointsByFrustum(FConvexVolume(Planes));
+				Normals[i] = -Normals[i];
 			}
-			else
+		}
+
+		// Perspective View
+		if (View->IsPerspectiveProjection())
+		{
+			Origins.Last() = Origins[0];
+		}
+		// Ortho Views
+		else
+		{
+			// Adjust for the LocationOffset
+			for (int32 i = 0; i < Points.Num(); ++i)
 			{
-				Editor->SelectPointsByFrustum(FConvexVolume(Planes), Mode == ELidarPointCloudSelectionMode::Add);
+				Origins[i] -= LocationOffset;
 			}
+
+			Origins.Last() = -Origins.Last(1);
+		}
+
+		for (int32 i = 0; i < Origins.Num(); ++i)
+		{
+			ConvexVolume.Planes.Emplace(Origins[i], Normals[i]);
 		}
+
+		ConvexVolume.Init();
 	}
+
+	return ConvexVolume;
 }
 
 void FLidarPointCloudEditorViewportClient::ResetCamera()
@@ -493,4 +1028,14 @@ bool FLidarPointCloudEditorViewportClient::IsSetShowNodesChecked() const
 	return false;
 }
 
+void FLidarPointCloudEditorViewportClient::SetSelectionMethod(ELidarPointCloudSelectionMethod NewSelectionMethod)
+{
+	SelectionMethod = NewSelectionMethod;
+
+	if (TSharedPtr<SLidarPointCloudEditorViewport> EditorViewport = PointCloudEditorViewportPtr.Pin())
+	{
+		EditorViewport->PaintBrush->SetVisibility(SelectionMethod == ELidarPointCloudSelectionMethod::Paint);
+	}
+}
+
 #undef LOCTEXT_NAMESPACE 
diff --git a/Engine/Plugins/Enterprise/LidarPointCloud/Source/LidarPointCloudEditor/Private/LidarPointCloudEditorViewportClient.h b/Engine/Plugins/Enterprise/LidarPointCloud/Source/LidarPointCloudEditor/Private/LidarPointCloudEditorViewportClient.h
--- a/Engine/Plugins/Enterprise/LidarPointCloud/Source/LidarPointCloudEditor/Private/LidarPointCloudEditorViewportClient.h
+++ b/Engine/Plugins/Enterprise/LidarPointCloud/Source/LidarPointCloudEditor/Private/LidarPointCloudEditorViewportClient.h
@@ -7,6 +7,7 @@
 #include "UnrealWidgetFwd.h"
 #include "EditorViewportClient.h"
 #include "Components.h"
+#include "LidarPointCloudShared.h"
 
 class FAdvancedPreviewScene;
 class FCanvas;
@@ -19,8 +20,15 @@ enum class ELidarPointCloudSelectionMode : uint8
 {
 	None,
 	Add,
-	Subtract,
-	Replace
+	Subtract
+};
+
+enum class ELidarPointCloudSelectionMethod : uint8
+{
+	Box,
+	Polygonal,
+	Lasso,
+	Paint
 };
 
 /** Viewport Client for the preview viewport */
@@ -31,12 +39,9 @@ public:
 	~FLidarPointCloudEditorViewportClient();
 
 	// FEditorViewportClient interface
-	virtual void MouseMove(FViewport* Viewport, int32 x, int32 y) override;
-	virtual bool InputKey(FViewport* Viewport, int32 ControllerId, FKey Key, EInputEvent Event, float AmountDepressed = 1.f, bool bGamepad = false) override;
-	virtual bool InputAxis(FViewport* Viewport, int32 ControllerId, FKey Key, float Delta, float DeltaTime, int32 NumSamples = 1, bool bGamepad = false) override;
-	virtual void ProcessClick(class FSceneView& View, class HHitProxy* HitProxy, FKey Key, EInputEvent Event, uint32 HitX, uint32 HitY) override;
+	virtual bool InputKey(FViewport* InViewport, int32 ControllerId, FKey Key, EInputEvent Event, float AmountDepressed = 1.f, bool bGamepad = false) override;
+	virtual bool InputAxis(FViewport* InViewport, int32 ControllerId, FKey Key, float Delta, float DeltaTime, int32 NumSamples = 1, bool bGamepad = false) override;
 	virtual void Tick(float DeltaSeconds) override;
-	virtual void Draw(const FSceneView* View, FPrimitiveDrawInterface* PDI) override;
 	virtual void DrawCanvas(FViewport& InViewport, FSceneView& View, FCanvas& Canvas) override;
 	virtual bool InputWidgetDelta(FViewport* InViewport, EAxisList::Type CurrentAxis, FVector& Drag, FRotator& Rot, FVector& Scale) override { return false; }
 	virtual void TrackingStarted(const struct FInputEventState& InInputState, bool bIsDragging, bool bNudge) override {}
@@ -49,6 +54,8 @@ public:
 	virtual FMatrix GetWidgetCoordSystem() const override { return FMatrix::Identity; }
 	virtual ECoordSystem GetWidgetCoordSystemSpace() const override { return COORD_Local; }
 	virtual bool ShouldOrbitCamera() const override;
+	virtual void LostFocus(FViewport* InViewport) override;
+	virtual void ReceivedFocus(FViewport* InViewport) override;
 
 	void ResetCamera();
 
@@ -58,6 +65,9 @@ public:
 	/** Callback for checking the nodes show flag. */
 	bool IsSetShowNodesChecked() const;
 
+	FORCEINLINE ELidarPointCloudSelectionMethod GetSelectionMethod() const { return SelectionMethod; }
+	void SetSelectionMethod(ELidarPointCloudSelectionMethod NewSelectionMethod);
+
 protected:
 	// FEditorViewportClient interface
 	virtual void PerspectiveCameraMoved() override;
@@ -67,9 +77,21 @@ protected:
 	/** Used to (re)-set the viewport show flags related to post processing*/
 	void SetAdvancedShowFlagsForScene(const bool bAdvancedShowFlags);
 
+	FSceneView* GetView();
+	FLidarPointCloudRay DeprojectCurrentMousePosition();
+
 private:
-	void OnSelectionStart();
-	void OnSelectionEnd(const FIntVector4& SelectionArea, ELidarPointCloudSelectionMode Mode);
+	void OnBoxSelectionEnd();
+	void OnPolygonalSelectionEnd();
+	void OnLassoSelectionEnd();
+	void OnPaintSelection();
+
+	void DrawSelectionBox(FCanvas& Canvas);
+	void DrawSelectionPolygonal(FCanvas& Canvas);
+	void DrawSelectionLasso(FCanvas& Canvas);
+	void DrawSelectionPaint(FCanvas& Canvas);
+
+	FConvexVolume BuildConvexVolumeForPoints(const TArray<FVector2D>& Points);
 
 private:
 	/** Component for the point cloud. */
@@ -84,9 +106,12 @@ private:
 	/** Stored pointer to the preview scene in which the point cloud is shown */
 	FAdvancedPreviewScene* AdvancedPreviewScene;
 
-private:
+	ELidarPointCloudSelectionMethod SelectionMethod;
 	ELidarPointCloudSelectionMode SelectionMode;
+	TArray<FIntPoint> SelectionPoints;
 
-	/** Location of the mouse cursor when the selection has started */
-	FIntPoint SelectionStartLocation;
+	float PaintingRadius;
+	bool bLineTraceHit;
+	FVector LineTraceHitPoint;
+	float LineTraceDistance;
 };
diff --git a/Engine/Plugins/Enterprise/LidarPointCloud/Source/LidarPointCloudEditor/Private/LidarPointCloudStyle.cpp b/Engine/Plugins/Enterprise/LidarPointCloud/Source/LidarPointCloudEditor/Private/LidarPointCloudStyle.cpp
--- a/Engine/Plugins/Enterprise/LidarPointCloud/Source/LidarPointCloudEditor/Private/LidarPointCloudStyle.cpp
+++ b/Engine/Plugins/Enterprise/LidarPointCloud/Source/LidarPointCloudEditor/Private/LidarPointCloudStyle.cpp
@@ -53,6 +53,11 @@ void FLidarPointCloudStyle::Initialize()
 
 	StyleSet->Set("LidarPointCloudEditor.EditMode", new IMAGE_PLUGIN_BRUSH("icon_edit_40", Icon40x40));
 	
+	StyleSet->Set("LidarPointCloudEditor.BoxSelection", new IMAGE_PLUGIN_BRUSH("icon_selbox_40", Icon40x40));
+	StyleSet->Set("LidarPointCloudEditor.PolygonalSelection", new IMAGE_PLUGIN_BRUSH("icon_selpoly_40", Icon40x40));
+	StyleSet->Set("LidarPointCloudEditor.LassoSelection", new IMAGE_PLUGIN_BRUSH("icon_sellasso_40", Icon40x40));
+	StyleSet->Set("LidarPointCloudEditor.PaintSelection", new IMAGE_PLUGIN_BRUSH("icon_selpaint_40", Icon40x40));
+
 	StyleSet->Set("LidarPointCloudEditor.InvertSelection", new IMAGE_PLUGIN_BRUSH("icon_invsel_40", Icon40x40));
 
 	StyleSet->Set("LidarPointCloudEditor.HideSelected", new IMAGE_PLUGIN_BRUSH("icon_hideselected_40", Icon40x40));
diff --git a/Engine/Plugins/Enterprise/LidarPointCloud/Source/LidarPointCloudRuntime/Private/IO/LidarPointCloudFileIO_E57.cpp b/Engine/Plugins/Enterprise/LidarPointCloud/Source/LidarPointCloudRuntime/Private/IO/LidarPointCloudFileIO_E57.cpp
--- a/Engine/Plugins/Enterprise/LidarPointCloud/Source/LidarPointCloudRuntime/Private/IO/LidarPointCloudFileIO_E57.cpp
+++ b/Engine/Plugins/Enterprise/LidarPointCloud/Source/LidarPointCloudRuntime/Private/IO/LidarPointCloudFileIO_E57.cpp
@@ -212,7 +212,7 @@ bool ULidarPointCloudFileIO_E57::HandleImport(const FString& Filename, TSharedPt
 			FLidarPointCloudDataBuffer* ThreadBuffer = BufferManager.GetFreeBuffer();
 			FMemory::Memcpy(ThreadBuffer->GetData(), Buffer->GetData(), NumPointsRead * sizeof(FE57Point));
 
-			ThreadResults.Add(Async(EAsyncExecution::ThreadPool, [NumPointsRead, bHasRGB, bHasIntensity, Transform, &OutImportResults, ImportScale, &SyncLock, ThreadBuffer]{
+			ThreadResults.Add(Async(EAsyncExecution::Thread, [NumPointsRead, bHasRGB, bHasIntensity, Transform, &OutImportResults, ImportScale, &SyncLock, ThreadBuffer]{
 				TArray64<FLidarPointCloudPoint> Points;
 				Points.Reserve(NumPointsRead);
 				FBox Bounds(EForceInit::ForceInit);
diff --git a/Engine/Plugins/Enterprise/LidarPointCloud/Source/LidarPointCloudRuntime/Private/IO/LidarPointCloudFileIO_LAS.cpp b/Engine/Plugins/Enterprise/LidarPointCloud/Source/LidarPointCloudRuntime/Private/IO/LidarPointCloudFileIO_LAS.cpp
--- a/Engine/Plugins/Enterprise/LidarPointCloud/Source/LidarPointCloudRuntime/Private/IO/LidarPointCloudFileIO_LAS.cpp
+++ b/Engine/Plugins/Enterprise/LidarPointCloud/Source/LidarPointCloudRuntime/Private/IO/LidarPointCloudFileIO_LAS.cpp
@@ -623,7 +623,7 @@ bool ULidarPointCloudFileIO_LAS::HandleImportLAS(const FString& Filename, FLidar
 			// Read the data
 			Reader->Serialize(Buffer->GetData(), PointsToRead * Header.PointDataRecordLength);
 
-			ThreadResults.Add(Async(EAsyncExecution::ThreadPool, [PointsToRead, Buffer, &ImportScale, &bHasIntensityData, IntensityBitShift, RGBBitShift, &Header, &OutImportResults, &PointsLock, bHasRGB, RGBOffset, ClassificationOffset, bUseConcurrentImport]
+			ThreadResults.Add(Async(EAsyncExecution::Thread, [PointsToRead, Buffer, &ImportScale, &bHasIntensityData, IntensityBitShift, RGBBitShift, &Header, &OutImportResults, &PointsLock, bHasRGB, RGBOffset, ClassificationOffset, bUseConcurrentImport]
 				{
 					uint8* Data = Buffer->GetData();
 
@@ -828,7 +828,7 @@ bool ULidarPointCloudFileIO_LAS::HandleImportLAZ(const FString& Filename, FLidar
 
 			Wrapper.Seek(StartIdx);
 
-			ThreadResults.Add(Async(EAsyncExecution::ThreadPool, [t, StartIdx, EndIdx, &Wrapper, &Point, &OutImportResults, bHasIntensityData, IntensityBitShift, bHasRGB, RGBBitShift, ImportScale, Header, &PointsLock, bUseConcurrentImport]
+			ThreadResults.Add(Async(EAsyncExecution::Thread, [t, StartIdx, EndIdx, &Wrapper, &Point, &OutImportResults, bHasIntensityData, IntensityBitShift, bHasRGB, RGBBitShift, ImportScale, Header, &PointsLock, bUseConcurrentImport]
 			{
 				// Local batching to avoid spiking RAM
 				const int64 BatchSize = 1000000;
diff --git a/Engine/Plugins/Enterprise/LidarPointCloud/Source/LidarPointCloudRuntime/Private/LidarPointCloud.cpp b/Engine/Plugins/Enterprise/LidarPointCloud/Source/LidarPointCloudRuntime/Private/LidarPointCloud.cpp
--- a/Engine/Plugins/Enterprise/LidarPointCloud/Source/LidarPointCloudRuntime/Private/LidarPointCloud.cpp
+++ b/Engine/Plugins/Enterprise/LidarPointCloud/Source/LidarPointCloudRuntime/Private/LidarPointCloud.cpp
@@ -11,6 +11,7 @@
 #include "Serialization/CustomVersion.h"
 #include "Misc/ScopeTryLock.h"
 #include "Misc/ConfigCacheIni.h"
+#include "Misc/ScopedSlowTask.h"
 #include "Engine/Engine.h"
 #include "LatentActions.h"
 #include "PhysicsEngine/BodySetup.h"
@@ -273,6 +274,7 @@ ULidarPointCloud::ULidarPointCloud()
 	: MaxCollisionError(100)
 	, NormalsQuality(40)
 	, NormalsNoiseTolerance(1)
+	, bOptimizedForDynamicData(false)
 	, Octree(this)
 	, OriginalCoordinates(FDoubleVector::ZeroVector)
 	, LocationOffset(FDoubleVector::ZeroVector)
@@ -397,6 +399,11 @@ void ULidarPointCloud::PostEditChangeProperty(FPropertyChangedEvent& PropertyCha
 			SetSourcePath(SourcePath.FilePath);
 		}
 
+		if (IS_PROPERTY(bOptimizedForDynamicData))
+		{
+			SetOptimizedForDynamicData(bOptimizedForDynamicData);
+		}
+
 		if (IS_PROPERTY(MaxCollisionError))
 		{
 			if (MaxCollisionError < Octree.GetEstimatedPointSpacing())
@@ -464,14 +471,24 @@ void ULidarPointCloud::GetPointsInBox(TArray<FLidarPointCloudPoint*>& SelectedPo
 	GetPointsInBox_Internal(SelectedPoints, Box, bVisibleOnly);
 }
 
+void ULidarPointCloud::GetPointsInConvexVolume(TArray64<FLidarPointCloudPoint*>& SelectedPoints, const FConvexVolume& ConvexVolume, const bool& bVisibleOnly)
+{
+	GetPointsInConvexVolume_Internal(SelectedPoints, ConvexVolume, bVisibleOnly);
+}
+
+void ULidarPointCloud::GetPointsInConvexVolume(TArray<FLidarPointCloudPoint*>& SelectedPoints, const FConvexVolume& ConvexVolume, const bool& bVisibleOnly)
+{
+	GetPointsInConvexVolume_Internal(SelectedPoints, ConvexVolume, bVisibleOnly);
+}
+
 void ULidarPointCloud::GetPointsInFrustum(TArray64<FLidarPointCloudPoint*>& SelectedPoints, const FConvexVolume& Frustum, const bool& bVisibleOnly)
 {
-	GetPointsInFrustum_Internal(SelectedPoints, Frustum, bVisibleOnly);
+	GetPointsInConvexVolume_Internal(SelectedPoints, Frustum, bVisibleOnly);
 }
 
 void ULidarPointCloud::GetPointsInFrustum(TArray<FLidarPointCloudPoint*>& SelectedPoints, const FConvexVolume& Frustum, const bool& bVisibleOnly)
 {
-	GetPointsInFrustum_Internal(SelectedPoints, Frustum, bVisibleOnly);
+	GetPointsInConvexVolume_Internal(SelectedPoints, Frustum, bVisibleOnly);
 }
 
 TArray<FLidarPointCloudPoint> ULidarPointCloud::GetPointsAsCopies(bool bReturnWorldSpace, int32 StartIndex, int32 Count) const
@@ -558,6 +575,25 @@ void ULidarPointCloud::SetSourcePath(const FString& NewSourcePath)
 	}
 }
 
+void ULidarPointCloud::SetOptimizedForDynamicData(bool bNewOptimizedForDynamicData)
+{
+#if WITH_EDITOR
+	FScopedSlowTask ProgressDialog(1, LOCTEXT("Optimize", "Optimizing Point Clouds..."));
+	ProgressDialog.MakeDialog();
+#endif
+
+	bOptimizedForDynamicData = bNewOptimizedForDynamicData;
+
+	if (bOptimizedForDynamicData)
+	{
+		Octree.OptimizeForDynamicData();
+	}
+	else
+	{
+		Octree.OptimizeForStaticData();
+	}
+}
+
 void ULidarPointCloud::BuildCollision()
 {
 	if (bCollisionBuildInProgress)
@@ -637,6 +673,9 @@ void ULidarPointCloud::Reimport(const FLidarPointCloudAsyncParameters& AsyncPara
 
 		const bool bCenter = GetDefault<ULidarPointCloudSettings>()->bAutoCenterOnImport;
 
+		// Reset the data optimization flag on import
+		bOptimizedForDynamicData = false;
+
 		// The actual import function to be executed
 		auto ImportFunction = [this, Notification, AsyncParameters, bCenter]
 		{
@@ -873,72 +912,86 @@ bool ULidarPointCloud::InsertPoints_NoLock(FLidarPointCloudPoint** InPoints, con
 template<typename T>
 bool ULidarPointCloud::InsertPoints_NoLock_Internal(T InPoints, const int64& Count, ELidarPointCloudDuplicateHandling DuplicateHandling, bool bRefreshPointsBounds, const FVector& Translation, FThreadSafeBool* bCanceled, TFunction<void(float)> ProgressCallback)
 {
-	const int32 MaxBatchSize = GetDefault<ULidarPointCloudSettings>()->MultithreadingInsertionBatchSize;
-	
-	// Minimum amount of points to progress to count as 1%
-	int64 RefreshStatusFrequency = Count * 0.01f;
-	FThreadSafeCounter64 ProcessedPoints(0);
-	int64 TotalProcessedPoints = 0;
-
-	const int32 NumThreads = FMath::Min(FPlatformMisc::NumberOfCoresIncludingHyperthreads() - 1, (int32)(Count / MaxBatchSize) + 1);
-	TArray<TFuture<void>>ThreadResults;
-	ThreadResults.Reserve(NumThreads);
-	const int64 NumPointsPerThread = Count / NumThreads + 1;
-
-	FCriticalSection ProgressCallbackLock;
-
-	// Fire threads
-	for (int32 ThreadID = 0; ThreadID < NumThreads; ThreadID++)
+	if (bOptimizedForDynamicData)
 	{
-		ThreadResults.Add(Async(EAsyncExecution::Thread, [this, ThreadID, DuplicateHandling, bRefreshPointsBounds, MaxBatchSize, NumPointsPerThread, RefreshStatusFrequency, &ProcessedPoints, &TotalProcessedPoints, InPoints, Count, &ProgressCallback, &ProgressCallbackLock, &bCanceled, &Translation]
+		Octree.InsertPoints(InPoints, Count, DuplicateHandling, bRefreshPointsBounds, Translation);
+		
+		if (ProgressCallback)
 		{
-			int64 Idx = ThreadID * NumPointsPerThread;
-			int64 MaxIdx = FMath::Min(Idx + NumPointsPerThread, Count);
-			T DataPointer = InPoints + Idx;
+			ProgressCallback(1.0);
+		}
+		
+		return true;
+	}
+	else
+	{
+		const int32 MaxBatchSize = GetDefault<ULidarPointCloudSettings>()->MultithreadingInsertionBatchSize;
 
-			while (Idx < MaxIdx)
-			{
-				int32 BatchSize = FMath::Min(MaxIdx - Idx, (int64)MaxBatchSize);
+		// Minimum amount of points to progress to count as 1%
+		int64 RefreshStatusFrequency = Count * 0.01f;
+		FThreadSafeCounter64 ProcessedPoints(0);
+		int64 TotalProcessedPoints = 0;
+
+		const int32 NumThreads = FMath::Min(FPlatformMisc::NumberOfCoresIncludingHyperthreads() - 1, (int32)(Count / MaxBatchSize) + 1);
+		TArray<TFuture<void>>ThreadResults;
+		ThreadResults.Reserve(NumThreads);
+		const int64 NumPointsPerThread = Count / NumThreads + 1;
 
-				Octree.InsertPoints(DataPointer, BatchSize, DuplicateHandling, bRefreshPointsBounds, Translation);
+		FCriticalSection ProgressCallbackLock;
 
-				if (ProgressCallback)
+		// Fire threads
+		for (int32 ThreadID = 0; ThreadID < NumThreads; ThreadID++)
+		{
+			ThreadResults.Add(Async(EAsyncExecution::Thread, [this, ThreadID, DuplicateHandling, bRefreshPointsBounds, MaxBatchSize, NumPointsPerThread, RefreshStatusFrequency, &ProcessedPoints, &TotalProcessedPoints, InPoints, Count, &ProgressCallback, &ProgressCallbackLock, &bCanceled, &Translation]
 				{
-					ProcessedPoints.Add(BatchSize);
-					if (ProcessedPoints.GetValue() > RefreshStatusFrequency)
+					int64 Idx = ThreadID * NumPointsPerThread;
+					int64 MaxIdx = FMath::Min(Idx + NumPointsPerThread, Count);
+					T DataPointer = InPoints + Idx;
+
+					while (Idx < MaxIdx)
 					{
-						FScopeLock Lock(&ProgressCallbackLock);
-						TotalProcessedPoints += ProcessedPoints.GetValue();
-						ProcessedPoints.Reset();
-						ProgressCallback((double)TotalProcessedPoints / Count);
-					}
-				}
+						int32 BatchSize = FMath::Min(MaxIdx - Idx, (int64)MaxBatchSize);
 
-				if (bCanceled && *bCanceled)
-				{
-					break;
-				}
+						Octree.InsertPoints(DataPointer, BatchSize, DuplicateHandling, bRefreshPointsBounds, Translation);
 
-				Idx += BatchSize;
-				DataPointer += BatchSize;
-			}
-		}));
-	}
+						if (ProgressCallback)
+						{
+							ProcessedPoints.Add(BatchSize);
+							if (ProcessedPoints.GetValue() > RefreshStatusFrequency)
+							{
+								FScopeLock Lock(&ProgressCallbackLock);
+								TotalProcessedPoints += ProcessedPoints.GetValue();
+								ProcessedPoints.Reset();
+								ProgressCallback((double)TotalProcessedPoints / Count);
+							}
+						}
 
-	// Sync
-	for (const TFuture<void>& ThreadResult : ThreadResults)
-	{
-		ThreadResult.Get();
-	}
+						if (bCanceled && *bCanceled)
+						{
+							break;
+						}
 
-	// Do not attempt to touch Render Data if being destroyed
-	if (!HasAnyFlags(RF_BeginDestroyed))
-	{
-		// Update PointCloudAssetRegistryCache
-		PointCloudAssetRegistryCache.PointCount = FString::FromInt(Octree.GetNumPoints());
-	}
+						Idx += BatchSize;
+						DataPointer += BatchSize;
+					}
+				}));
+		}
 
-	return !bCanceled || !(*bCanceled);
+		// Sync
+		for (const TFuture<void>& ThreadResult : ThreadResults)
+		{
+			ThreadResult.Get();
+		}
+
+		// Do not attempt to touch Render Data if being destroyed
+		if (!HasAnyFlags(RF_BeginDestroyed))
+		{
+			// Update PointCloudAssetRegistryCache
+			PointCloudAssetRegistryCache.PointCount = FString::FromInt(Octree.GetNumPoints());
+		}
+
+		return !bCanceled || !(*bCanceled);
+	}
 }
 
 template<typename T>
@@ -1367,9 +1420,9 @@ void ULidarPointCloud::GetPointsInBox_Internal(TArray<FLidarPointCloudPoint*, T>
 }
 
 template <typename T>
-void ULidarPointCloud::GetPointsInFrustum_Internal(TArray<FLidarPointCloudPoint*, T>& SelectedPoints, const FConvexVolume& Frustum, const bool& bVisibleOnly)
+void ULidarPointCloud::GetPointsInConvexVolume_Internal(TArray<FLidarPointCloudPoint*, T>& SelectedPoints, const FConvexVolume& ConvexVolume, const bool& bVisibleOnly)
 {
-	Octree.GetPointsInFrustum(SelectedPoints, Frustum, bVisibleOnly);
+	Octree.GetPointsInConvexVolume(SelectedPoints, ConvexVolume, bVisibleOnly);
 }
 
 template <typename T>
diff --git a/Engine/Plugins/Enterprise/LidarPointCloud/Source/LidarPointCloudRuntime/Private/LidarPointCloudComponent.cpp b/Engine/Plugins/Enterprise/LidarPointCloud/Source/LidarPointCloudRuntime/Private/LidarPointCloudComponent.cpp
--- a/Engine/Plugins/Enterprise/LidarPointCloud/Source/LidarPointCloudRuntime/Private/LidarPointCloudComponent.cpp
+++ b/Engine/Plugins/Enterprise/LidarPointCloud/Source/LidarPointCloudRuntime/Private/LidarPointCloudComponent.cpp
@@ -19,10 +19,8 @@
 
 ULidarPointCloudComponent::ULidarPointCloudComponent()
 	: CustomMaterial(nullptr)
-	, MinScreenSize(0.05f)
-	, bUseFrustumCulling(true)
 	, PointSize(1.0f)
-	, bUseScreenSizeScaling(false)
+	, ScalingMethod(ELidarPointCloudScalingMethod::PerNodeAdaptive)
 	, GapFillingStrength(0)
 	, ColorSource(ELidarPointCloudColorationMode::Data)
 	, PointShape(ELidarPointCloudSpriteShape::Square)
@@ -37,6 +35,7 @@ ULidarPointCloudComponent::ULidarPointCloudComponent()
 	, Offset(FVector::ZeroVector)
 	, ColorTint(FLinearColor::White)
 	, IntensityInfluence(0.0f)
+	, bUseFrustumCulling(true)
 	, MinDepth(0)
 	, MaxDepth(-1)
 	, bDrawNodeBounds(false)
@@ -242,3 +241,47 @@ void ULidarPointCloudComponent::PostEditChangeProperty(FPropertyChangedEvent& Pr
 	Super::PostEditChangeProperty(PropertyChangedEvent);
 }
 #endif
+
+void FLidarPointCloudComponentRenderParams::UpdateFromComponent(ULidarPointCloudComponent* Component)
+{
+	MinDepth = Component->MinDepth;
+	MaxDepth = Component->MaxDepth;
+
+	BoundsScale = Component->BoundsScale;
+	BoundsSize = Component->GetPointCloud()->GetBounds().GetSize();
+
+	// Make sure to apply minimum bounds size
+	BoundsSize.X = FMath::Max(BoundsSize.X, 0.001f);
+	BoundsSize.Y = FMath::Max(BoundsSize.Y, 0.001f);
+	BoundsSize.Z = FMath::Max(BoundsSize.Z, 0.001f);
+
+	LocationOffset = Component->GetPointCloud()->GetLocationOffset().ToVector();
+	ComponentScale = Component->GetComponentScale().GetAbsMax();
+
+	PointSize = Component->PointSize;
+	PointSizeBias = Component->PointSizeBias;
+	GapFillingStrength = Component->GapFillingStrength;
+
+	bOwnedByEditor = Component->IsOwnedByEditor();
+	bDrawNodeBounds = Component->bDrawNodeBounds;
+	bShouldRenderFacingNormals = Component->ShouldRenderFacingNormals();
+	bUseFrustumCulling = Component->bUseFrustumCulling;
+
+	ScalingMethod = Component->ScalingMethod;
+
+	ColorSource = Component->ColorSource;
+	PointShape = Component->GetPointShape();
+
+	Offset = Component->Offset;
+	Contrast = Component->Contrast;
+	Saturation = Component->Saturation;
+	Gamma = Component->Gamma;
+	ColorTint = FVector(Component->ColorTint);
+	IntensityInfluence = Component->IntensityInfluence;
+
+	ClassificationColors = Component->ClassificationColors;
+	ElevationColorBottom = Component->ElevationColorBottom;
+	ElevationColorTop = Component->ElevationColorTop;
+
+	Material = Component->GetMaterial(0);
+}
diff --git a/Engine/Plugins/Enterprise/LidarPointCloud/Source/LidarPointCloudRuntime/Private/LidarPointCloudLODManager.cpp b/Engine/Plugins/Enterprise/LidarPointCloud/Source/LidarPointCloudRuntime/Private/LidarPointCloudLODManager.cpp
--- a/Engine/Plugins/Enterprise/LidarPointCloud/Source/LidarPointCloudRuntime/Private/LidarPointCloudLODManager.cpp
+++ b/Engine/Plugins/Enterprise/LidarPointCloud/Source/LidarPointCloudRuntime/Private/LidarPointCloudLODManager.cpp
@@ -19,15 +19,12 @@
 #include "EditorViewportClient.h"
 #endif
 
-FCriticalSection ProxyComponentMapLock;
-TMap<TWeakPtr<FLidarPointCloudSceneProxyWrapper, ESPMode::ThreadSafe>, TWeakObjectPtr<ULidarPointCloudComponent>> ProxyComponentMap;
-
 DECLARE_CYCLE_STAT(TEXT("Node Selection"), STAT_NodeSelection, STATGROUP_LidarPointCloud)
 DECLARE_CYCLE_STAT(TEXT("Node Processing"), STAT_NodeProcessing, STATGROUP_LidarPointCloud)
 DECLARE_CYCLE_STAT(TEXT("Render Data Update"), STAT_UpdateRenderData, STATGROUP_LidarPointCloud)
+DECLARE_CYCLE_STAT(TEXT("Node Streaming"), STAT_NodeStreaming, STATGROUP_LidarPointCloud);
 
 DECLARE_DWORD_ACCUMULATOR_STAT(TEXT("Total Point Count [thousands]"), STAT_PointCountTotal, STATGROUP_LidarPointCloud)
-DECLARE_DWORD_ACCUMULATOR_STAT(TEXT("Points In Frustum"), STAT_PointCountFrustum, STATGROUP_LidarPointCloud)
 DECLARE_DWORD_ACCUMULATOR_STAT(TEXT("Point Budget"), STAT_PointBudget, STATGROUP_LidarPointCloud)
 DECLARE_DWORD_ACCUMULATOR_STAT(TEXT("Visible Points"), STAT_PointCount, STATGROUP_LidarPointCloud)
 DECLARE_DWORD_ACCUMULATOR_STAT(TEXT("Loaded Nodes"), STAT_LoadedNodes, STATGROUP_LidarPointCloud)
@@ -175,138 +172,138 @@ bool FLidarPointCloudViewData::ComputeFromEditorViewportClient(FViewportClient*
 	return false;
 }
 
-void FLidarPointCloudTraversalOctree::GetVisibleNodes(TArray<FLidarPointCloudLODManager::FNodeSizeData>& NodeSizeData, const FLidarPointCloudViewData& ViewData, const int32& ProxyIndex, const FLidarPointCloudNodeSelectionParams& SelectionParams, const float& CurrentTime)
+int32 FLidarPointCloudTraversalOctree::GetVisibleNodes(TArray<FLidarPointCloudTraversalOctreeNodeSizeData>& NodeSizeData, const FLidarPointCloudViewData* ViewData, const int32& ProxyIndex, const FLidarPointCloudNodeSelectionParams& SelectionParams)
 {
-	// Skip processing, if the asset is not visible
-	if (!ViewData.ViewFrustum.IntersectBox(GetCenter(), GetExtent()))
-	{
-		return;
-	}
-
-	float MinScreenSizeSq = SelectionParams.MinScreenSize * SelectionParams.MinScreenSize;
-	float BoundsScaleSq = SelectionParams.BoundsScale * SelectionParams.BoundsScale;
+	const int32 CurrentNodeCount = NodeSizeData.Num();
 
-	const float BaseLODImportance = FMath::Max(0.0f, CVarBaseLODImportance.GetValueOnAnyThread());
-
-	bool bStartClipped = false;
-	if (SelectionParams.ClippingVolumes)
+	// Only process, if the asset is visible
+	if (ViewData->ViewFrustum.IntersectBox(GetCenter(), GetExtent()))
 	{
-		for (const ALidarClippingVolume* Volume : *SelectionParams.ClippingVolumes)
+		const float BoundsScaleSq = SelectionParams.BoundsScale * SelectionParams.BoundsScale;
+		const float BaseLODImportance = FMath::Max(0.0f, CVarBaseLODImportance.GetValueOnAnyThread());
+
+		bool bStartClipped = false;
+		if (SelectionParams.ClippingVolumes)
 		{
-			if (Volume->Mode == ELidarClippingVolumeMode::ClipOutside)
+			for (const FLidarPointCloudClippingVolumeParams& Volume : *SelectionParams.ClippingVolumes)
 			{
-				bStartClipped = true;
-				break;
+				if (Volume.Mode == ELidarClippingVolumeMode::ClipOutside)
+				{
+					bStartClipped = true;
+					break;
+				}
 			}
 		}
-	}
 
-	TQueue<FLidarPointCloudTraversalOctreeNode*> Nodes;
-	FLidarPointCloudTraversalOctreeNode* CurrentNode = nullptr;
-	Nodes.Enqueue(&Root);
-	while (Nodes.Dequeue(CurrentNode))
-	{
-		// Reset selection flag
-		CurrentNode->bSelected = false;
+		TQueue<FLidarPointCloudTraversalOctreeNode*> Nodes;
+		FLidarPointCloudTraversalOctreeNode* CurrentNode = nullptr;
+		Nodes.Enqueue(&Root);
+		while (Nodes.Dequeue(CurrentNode))
+		{
+			// Reset selection flag
+			CurrentNode->bSelected = false;
 
-		// Update number of visible points, if needed
-		CurrentNode->DataNode->UpdateNumVisiblePoints();
+			// Update number of visible points, if needed
+			CurrentNode->DataNode->UpdateNumVisiblePoints();
 
-		const FVector NodeExtent = Extents[CurrentNode->Depth] * SelectionParams.BoundsScale;
+			const FVector NodeExtent = Extents[CurrentNode->Depth] * SelectionParams.BoundsScale;
 
-		bool bFullyContained = true;
+			bool bFullyContained = true;
 
-		if (SelectionParams.bUseFrustumCulling && (CurrentNode->Depth == 0 || !CurrentNode->bFullyContained) && !ViewData.ViewFrustum.IntersectBox(CurrentNode->Center, NodeExtent, bFullyContained))
-		{
-			continue;
-		}
+			if (SelectionParams.bUseFrustumCulling && (CurrentNode->Depth == 0 || !CurrentNode->bFullyContained) && !ViewData->ViewFrustum.IntersectBox(CurrentNode->Center, NodeExtent, bFullyContained))
+			{
+				continue;
+			}
 
-		const FBox NodeBounds(CurrentNode->Center - NodeExtent, CurrentNode->Center + NodeExtent);
+			const FBox NodeBounds(CurrentNode->Center - NodeExtent, CurrentNode->Center + NodeExtent);
 
-		// Check vs Clipping Volumes
-		bool bClip = bStartClipped;
-		if (SelectionParams.ClippingVolumes)
-		{
-			for (const ALidarClippingVolume* Volume : *SelectionParams.ClippingVolumes)
+			// Check vs Clipping Volumes
+			bool bClip = bStartClipped;
+			if (SelectionParams.ClippingVolumes)
 			{
-				if (Volume->Mode == ELidarClippingVolumeMode::ClipOutside)
+				for (const FLidarPointCloudClippingVolumeParams& Volume : *SelectionParams.ClippingVolumes)
 				{
-					if (Volume->GetBounds().GetBox().Intersect(NodeBounds))
+					if (Volume.Mode == ELidarClippingVolumeMode::ClipOutside)
 					{
-						bClip = false;
+						if (Volume.Bounds.Intersect(NodeBounds))
+						{
+							bClip = false;
+						}
 					}
-				}
-				else
-				{
-					if (Volume->GetBounds().GetBox().IsInside(NodeBounds))
+					else
 					{
-						bClip = true;
+						if (Volume.Bounds.IsInside(NodeBounds))
+						{
+							bClip = true;
+						}
 					}
 				}
-			}
 
-			if (bClip)
-			{
-				continue;
+				if (bClip)
+				{
+					continue;
+				}
 			}
-		}
 
-		// Only process this node if it has any visible points - do not use continue; as the children may still contain visible points!
-		if (CurrentNode->DataNode->GetNumVisiblePoints() > 0 && CurrentNode->Depth >= SelectionParams.MinDepth)
-		{
-			float ScreenSizeSq = 0;
+			// Only process this node if it has any visible points - do not use continue; as the children may still contain visible points!
+			if (CurrentNode->DataNode->GetNumVisiblePoints() > 0 && CurrentNode->Depth >= SelectionParams.MinDepth)
+			{
+				float ScreenSizeSq = 0;
 
-			FVector VectorToNode = CurrentNode->Center - ViewData.ViewOrigin;
-			const float DistSq = VectorToNode.SizeSquared();
-			const float AdjustedRadiusSq = RadiiSq[CurrentNode->Depth] * BoundsScaleSq;
+				FVector VectorToNode = CurrentNode->Center - ViewData->ViewOrigin;
+				const float DistSq = VectorToNode.SizeSquared();
+				const float AdjustedRadiusSq = RadiiSq[CurrentNode->Depth] * BoundsScaleSq;
 
-			// Make sure to show at least the minimum depth for each visible asset
-			if (CurrentNode->Depth == SelectionParams.MinDepth)
-			{
-				// Add screen size to maintain hierarchy
-				ScreenSizeSq = BaseLODImportance + ViewData.ScreenSizeFactor * AdjustedRadiusSq / FMath::Max(1.0f, DistSq);
-			}
-			else
-			{
-				// If the camera is within this node's bounds, it should always be qualified for rendering
-				if (DistSq <= AdjustedRadiusSq)
+				// Make sure to show at least the minimum depth for each visible asset
+				if (CurrentNode->Depth == SelectionParams.MinDepth)
 				{
-					// Subtract Depth to maintain hierarchy 
-					ScreenSizeSq = 1000 - CurrentNode->Depth;
+					// Add screen size to maintain hierarchy
+					ScreenSizeSq = BaseLODImportance + ViewData->ScreenSizeFactor * AdjustedRadiusSq / FMath::Max(1.0f, DistSq);
 				}
 				else
 				{
-					ScreenSizeSq = ViewData.ScreenSizeFactor * AdjustedRadiusSq / FMath::Max(1.0f, DistSq);
-
-					// Check for minimum screen size
-					if (!ViewData.bSkipMinScreenSize && ScreenSizeSq < MinScreenSizeSq)
+					// If the camera is within this node's bounds, it should always be qualified for rendering
+					if (DistSq <= AdjustedRadiusSq)
 					{
-						continue;
+						// Subtract Depth to maintain hierarchy 
+						ScreenSizeSq = 1000 - CurrentNode->Depth;
 					}
-
-					// Add optional preferential selection for nodes closer to the screen center
-					if (SelectionParams.ScreenCenterImportance > 0)
+					else
 					{
-						VectorToNode.Normalize();
-						float Dot = FVector::DotProduct(ViewData.ViewDirection, VectorToNode);
+						ScreenSizeSq = ViewData->ScreenSizeFactor * AdjustedRadiusSq / FMath::Max(1.0f, DistSq);
+
+						// Check for minimum screen size
+						if (!ViewData->bSkipMinScreenSize && ScreenSizeSq < SelectionParams.MinScreenSize)
+						{
+							continue;
+						}
+
+						// Add optional preferential selection for nodes closer to the screen center
+						if (SelectionParams.ScreenCenterImportance > 0)
+						{
+							VectorToNode.Normalize();
+							float Dot = FVector::DotProduct(ViewData->ViewDirection, VectorToNode);
 
-						ScreenSizeSq = FMath::Lerp(ScreenSizeSq, ScreenSizeSq * Dot, SelectionParams.ScreenCenterImportance);
+							ScreenSizeSq = FMath::Lerp(ScreenSizeSq, ScreenSizeSq * Dot, SelectionParams.ScreenCenterImportance);
+						}
 					}
 				}
-			}
 
-			NodeSizeData.Emplace(CurrentNode, ScreenSizeSq, ProxyIndex);
-		}
+				NodeSizeData.Emplace(CurrentNode, ScreenSizeSq, ProxyIndex);
+			}
 
-		if (SelectionParams.MaxDepth < 0 || CurrentNode->Depth < SelectionParams.MaxDepth)
-		{
-			for (FLidarPointCloudTraversalOctreeNode& Child : CurrentNode->Children)
+			if (SelectionParams.MaxDepth < 0 || CurrentNode->Depth < SelectionParams.MaxDepth)
 			{
-				Child.bFullyContained = bFullyContained;
-				Nodes.Enqueue(&Child);
+				for (FLidarPointCloudTraversalOctreeNode& Child : CurrentNode->Children)
+				{
+					Child.bFullyContained = bFullyContained;
+					Nodes.Enqueue(&Child);
+				}
 			}
 		}
 	}
+
+	return NodeSizeData.Num() - CurrentNodeCount;
 }
 
 /** Calculates the correct point budget to use for current frame */
@@ -434,18 +431,10 @@ void FLidarPointCloudLODManager::Tick(float DeltaTime)
 
 	PrepareProxies();
 
-	// Gather clipping volumes
-	TArray<const ALidarClippingVolume*> ClippingVolumes = GetClippingVolumes();
-
-	// Sort clipping volumes by priority
-	Algo::Sort(ClippingVolumes, [](const ALidarClippingVolume* A, const ALidarClippingVolume* B) { return (A->Priority < B->Priority) || (A->Priority == B->Priority && A->Mode > B->Mode); });
-
-	// A copy of the array will be passed, to avoid concurrency issues
-	TArray<FRegisteredProxy> CurrentRegisteredProxies = RegisteredProxies;
-
-	Async(EAsyncExecution::ThreadPool, [this, CurrentRegisteredProxies, ClippingVolumes, PointBudget]
+	Async(EAsyncExecution::ThreadPool, [this, CurrentRegisteredProxies = RegisteredProxies, PointBudget, ClippingVolumes = GetClippingVolumes()]
 	{
 		NumPointsInFrustum.Set(ProcessLOD(CurrentRegisteredProxies, Time, PointBudget, ClippingVolumes));
+		bProcessing = false;
 	});
 }
 
@@ -459,19 +448,33 @@ void FLidarPointCloudLODManager::RegisterProxy(ULidarPointCloudComponent* Compon
 	if (IsValid(Component))
 	{
 		static FLidarPointCloudLODManager Instance;
+		static FCriticalSection ProxyLock;
 
 		if (TSharedPtr<FLidarPointCloudSceneProxyWrapper, ESPMode::ThreadSafe> SceneProxyWrapperShared = SceneProxyWrapper.Pin())
 		{
-			FScopeLock Lock(&ProxyComponentMapLock);
-
+			FScopeLock Lock(&ProxyLock);
 			Instance.RegisteredProxies.Emplace(Component, SceneProxyWrapper);
-			ProxyComponentMap.Add(SceneProxyWrapper, Component);
 		}
 	}
 }
 
-int64 FLidarPointCloudLODManager::ProcessLOD(const TArray<FLidarPointCloudLODManager::FRegisteredProxy>& InRegisteredProxies, const float CurrentTime, const uint32 PointBudget, const TArray<const ALidarClippingVolume*>& ClippingVolumes)
+int64 FLidarPointCloudLODManager::ProcessLOD(const TArray<FLidarPointCloudLODManager::FRegisteredProxy>& InRegisteredProxies, const float CurrentTime, const uint32 PointBudget, const TArray<FLidarPointCloudClippingVolumeParams>& ClippingVolumes)
 {
+	struct FSelectionFilterParams
+	{
+		float MinScreenSize;
+		uint32 NumNodes;
+	};
+
+	static TArray<FSelectionFilterParams> SelectionFilterParams;
+	{
+		const int32 DeltaSize = InRegisteredProxies.Num() - SelectionFilterParams.Num();
+		if (DeltaSize > 0)
+		{
+			SelectionFilterParams.AddZeroed(DeltaSize);
+		}
+	}
+
 	uint32 TotalPointsSelected = 0;
 	int64 NewNumPointsInFrustum = 0;
 
@@ -482,10 +485,9 @@ int64 FLidarPointCloudLODManager::ProcessLOD(const TArray<FLidarPointCloudLODMan
 		SCOPE_CYCLE_COUNTER(STAT_NodeSelection);
 
 		const float ScreenCenterImportance = CVarLidarScreenCenterImportance.GetValueOnAnyThread();
-
 		int32 NumSelectedNodes = 0;
 
-		TArray<FNodeSizeData> NodeSizeData;
+		TArray<FLidarPointCloudTraversalOctreeNodeSizeData> NodeSizeData;
 
 		for (int32 i = 0; i < InRegisteredProxies.Num(); ++i)
 		{
@@ -494,14 +496,6 @@ int64 FLidarPointCloudLODManager::ProcessLOD(const TArray<FLidarPointCloudLODMan
 			// Acquire a Shared Pointer from the Weak Pointer and check that it references a valid object
 			if (TSharedPtr<FLidarPointCloudSceneProxyWrapper, ESPMode::ThreadSafe> SceneProxyWrapper = RegisteredProxy.SceneProxyWrapper.Pin())
 			{
-				FScopeLock OctreeLock(&RegisteredProxy.PointCloud->Octree.DataLock);
-
-				// If the octree has been invalidated, skip processing
-				if (!RegisteredProxy.TraversalOctree->bValid)
-				{
-					continue;
-				}
-
 #if WITH_EDITOR
 				// Avoid doubling the point allocation of the same asset (once in Editor world and once in PIE world)
 				if (RegisteredProxy.bSkip)
@@ -510,29 +504,34 @@ int64 FLidarPointCloudLODManager::ProcessLOD(const TArray<FLidarPointCloudLODMan
 				}
 #endif
 
-				// Construct selection params
-				FLidarPointCloudNodeSelectionParams SelectionParams;
-				SelectionParams.MinScreenSize = FMath::Max(RegisteredProxy.Component->MinScreenSize, 0.0f);
-				SelectionParams.ScreenCenterImportance = ScreenCenterImportance;
-				SelectionParams.MinDepth = RegisteredProxy.Component->MinDepth;
-				SelectionParams.MaxDepth = RegisteredProxy.Component->MaxDepth;
-				SelectionParams.BoundsScale = RegisteredProxy.Component->BoundsScale;
-				SelectionParams.bUseFrustumCulling = RegisteredProxy.Component->bUseFrustumCulling;
-
-				// Ignore clipping if in editor viewport
-				SelectionParams.ClippingVolumes = RegisteredProxy.Component->IsOwnedByEditor() ? nullptr : &ClippingVolumes;
-
-				// Append visible nodes
-				RegisteredProxy.TraversalOctree->GetVisibleNodes(NodeSizeData, RegisteredProxy.ViewData, i, SelectionParams, CurrentTime);
+				// Attempt to get a data lock and check if the traversal octree is still valid
+				FScopeTryLock OctreeLock(&RegisteredProxy.Octree->DataLock);
+				if (OctreeLock.IsLocked() && RegisteredProxy.TraversalOctree->bValid)
+				{
+					FLidarPointCloudNodeSelectionParams SelectionParams;
+					SelectionParams.MinScreenSize = SelectionFilterParams[i].MinScreenSize;
+					SelectionParams.ScreenCenterImportance = ScreenCenterImportance;
+					SelectionParams.MinDepth = RegisteredProxy.ComponentRenderParams.MinDepth;
+					SelectionParams.MaxDepth = RegisteredProxy.ComponentRenderParams.MaxDepth;
+					SelectionParams.BoundsScale = RegisteredProxy.ComponentRenderParams.BoundsScale;
+					SelectionParams.bUseFrustumCulling = RegisteredProxy.ComponentRenderParams.bUseFrustumCulling;
+					SelectionParams.ClippingVolumes = RegisteredProxy.ComponentRenderParams.bOwnedByEditor ? nullptr : &ClippingVolumes; // Ignore clipping if in editor viewport
+
+					// Append visible nodes
+					SelectionFilterParams[i].NumNodes = RegisteredProxy.TraversalOctree->GetVisibleNodes(NodeSizeData, &RegisteredProxy.ViewData, i, SelectionParams);
+				}
 			}
 		}
 
 		// Sort Nodes
-		Algo::Sort(NodeSizeData, [](const FNodeSizeData& A, const FNodeSizeData& B) { return A.Size > B.Size; });
+		Algo::Sort(NodeSizeData, [](const FLidarPointCloudTraversalOctreeNodeSizeData& A, const FLidarPointCloudTraversalOctreeNodeSizeData& B) { return A.Size > B.Size; });
+
+		TArray<float> NewMinScreenSizes;
+		NewMinScreenSizes.AddZeroed(InRegisteredProxies.Num());
 
 		// Limit nodes using specified Point Budget
 		SelectedNodesData.AddDefaulted(InRegisteredProxies.Num());
-		for (FNodeSizeData& Element : NodeSizeData)
+		for (FLidarPointCloudTraversalOctreeNodeSizeData& Element : NodeSizeData)
 		{
 			const uint32 NumPoints = Element.Node->DataNode->GetNumVisiblePoints();
 			const uint32 NewNumPointsSelected = TotalPointsSelected + NumPoints;
@@ -540,102 +539,131 @@ int64 FLidarPointCloudLODManager::ProcessLOD(const TArray<FLidarPointCloudLODMan
 
 			if (NewNumPointsSelected <= PointBudget)
 			{
+				NewMinScreenSizes[Element.ProxyIndex] = FMath::Max(Element.Size, 0.0f);
 				SelectedNodesData[Element.ProxyIndex].Add(Element.Node);
 				TotalPointsSelected = NewNumPointsSelected;
 				Element.Node->bSelected = true;
+				--SelectionFilterParams[Element.ProxyIndex].NumNodes;
 				++NumSelectedNodes;
 			}
 		}
 
+		for (int32 i = 0; i < InRegisteredProxies.Num(); ++i)
+		{
+			// If point budget is saturated, apply new Min Screen Sizes
+			// Otherwise, decrease Min Screen Sizes, to allow for more points
+			SelectionFilterParams[i].MinScreenSize = SelectionFilterParams[i].NumNodes > 0 ? NewMinScreenSizes[i] : (SelectionFilterParams[i].MinScreenSize * 0.9f);
+		}
+
 		SET_DWORD_STAT(STAT_PointCount, TotalPointsSelected);
-		SET_DWORD_STAT(STAT_PointCountFrustum, NewNumPointsInFrustum);
 	}
 
 	// Used to pass render data updates to render thread
 	TArray<FLidarPointCloudProxyUpdateData> ProxyUpdateData;
 
+	// Holds a per-octree list of nodes to stream-in or extend
+	TMultiMap<FLidarPointCloudOctree*, TArray<FLidarPointCloudOctreeNode*>> OctreeStreamingMap;
+
 	// Process Nodes
 	{
 		SCOPE_CYCLE_COUNTER(STAT_NodeProcessing);
 
-		// Set when to release the BulkData, if no longer visible
-		const float BulkDataLifetime = CurrentTime + GetDefault<ULidarPointCloudSettings>()->CachedNodeLifetime;
+		const bool bUseStaticBuffers = GetDefault<ULidarPointCloudSettings>()->bUseFastRendering;
 
 		for (int32 i = 0; i < SelectedNodesData.Num(); ++i)
 		{
 			const FLidarPointCloudLODManager::FRegisteredProxy& RegisteredProxy = InRegisteredProxies[i];
 
-			FLidarPointCloudProxyUpdateData UpdateData;
-			UpdateData.SceneProxyWrapper = RegisteredProxy.SceneProxyWrapper;
-			UpdateData.NumElements = 0;
-			UpdateData.VDMultiplier = RegisteredProxy.TraversalOctree->ReversedVirtualDepthMultiplier;
-			UpdateData.RootCellSize = RegisteredProxy.PointCloud->Octree.GetRootCellSize();
-			UpdateData.ClippingVolumes = ClippingVolumes;
+			// Attempt to get a data lock and check if the traversal octree is still valid
+			FScopeTryLock OctreeLock(&RegisteredProxy.Octree->DataLock);
+			if (OctreeLock.IsLocked() && RegisteredProxy.TraversalOctree->bValid)
+			{
+				FLidarPointCloudProxyUpdateData UpdateData;
+				UpdateData.SceneProxyWrapper = RegisteredProxy.SceneProxyWrapper;
+				UpdateData.NumElements = 0;
+				UpdateData.VDMultiplier = RegisteredProxy.TraversalOctree->ReversedVirtualDepthMultiplier;
+				UpdateData.RootCellSize = RegisteredProxy.Octree->GetRootCellSize();
+				UpdateData.ClippingVolumes = ClippingVolumes;
+				UpdateData.bUseStaticBuffers = bUseStaticBuffers && !RegisteredProxy.Octree->IsOptimizedForDynamicData();
+				UpdateData.RenderParams = RegisteredProxy.ComponentRenderParams;
 
-			const bool bUseNormals = RegisteredProxy.Component->ShouldRenderFacingNormals();
+#if !(UE_BUILD_SHIPPING)
+				const bool bDrawNodeBounds = RegisteredProxy.ComponentRenderParams.bDrawNodeBounds;
+				if (bDrawNodeBounds)
+				{
+					UpdateData.Bounds.Reset(SelectedNodesData[i].Num());
+				}
+#endif
+				const bool bCalculateVirtualDepth = RegisteredProxy.ComponentRenderParams.PointSize > 0;
 
-			// Since the process is async, make sure we can access the data!
-			{
-				FScopeLock OctreeLock(&RegisteredProxy.PointCloud->Octree.DataLock);
-				
-				// If the octree has been invalidated, skip processing
-				if (!RegisteredProxy.TraversalOctree->bValid)
+				TArray<float> LocalLevelWeights;
+				TArray<float>* LevelWeightsPtr = &RegisteredProxy.TraversalOctree->LevelWeights;
+				float PointSizeBias = RegisteredProxy.ComponentRenderParams.PointSizeBias;
+
+				// Only calculate if needed
+				if (bCalculateVirtualDepth && (RegisteredProxy.ComponentRenderParams.ScalingMethod == ELidarPointCloudScalingMethod::PerNodeAdaptive || RegisteredProxy.ComponentRenderParams.ScalingMethod == ELidarPointCloudScalingMethod::PerPoint))
 				{
-					continue;
+					RegisteredProxy.TraversalOctree->CalculateLevelWeightsForSelectedNodes(LocalLevelWeights);
+					LevelWeightsPtr = &LocalLevelWeights;
+					PointSizeBias = 0;
 				}
 
 				// Queue nodes to be streamed
+				TArray<FLidarPointCloudOctreeNode*>& NodesToStream = OctreeStreamingMap.FindOrAdd(RegisteredProxy.Octree);
+				NodesToStream.Reserve(NodesToStream.Num() + SelectedNodesData[i].Num());
 				for (FLidarPointCloudTraversalOctreeNode* Node : SelectedNodesData[i])
 				{
-					RegisteredProxy.PointCloud->Octree.QueueNode(Node->DataNode, BulkDataLifetime);
+					NodesToStream.Add(Node->DataNode);
 
 					if (Node->DataNode->HasData())
 					{
-						uint8 VirtualDepth = 0;
-
 						// Only calculate if needed
-						if (RegisteredProxy.Component->PointSize > 0)
+                        if (bCalculateVirtualDepth)
+                        {
+							Node->CalculateVirtualDepth(*LevelWeightsPtr, PointSizeBias);
+                        }
+						
+						const uint32 NumVisiblePoints = Node->DataNode->GetNumVisiblePoints();
+						UpdateData.NumElements += NumVisiblePoints;
+						UpdateData.SelectedNodes.Emplace(bCalculateVirtualDepth ? Node->VirtualDepth : 0, NumVisiblePoints, Node->DataNode);
+#if !(UE_BUILD_SHIPPING)
+						if (bDrawNodeBounds)
 						{
-							VirtualDepth = Node->CalculateVirtualDepth(RegisteredProxy.TraversalOctree->LevelWeights, RegisteredProxy.TraversalOctree->VirtualDepthMultiplier, RegisteredProxy.Component->PointSizeBias);
+							const FVector Extent = RegisteredProxy.TraversalOctree->Extents[Node->Depth];
+							UpdateData.Bounds.Emplace(Node->Center - Extent, Node->Center + Extent);
 						}
-
-						UpdateData.SelectedNodes.Emplace(VirtualDepth, Node->DataNode->GetNumVisiblePoints(), Node->DataNode);
-						UpdateData.NumElements += Node->DataNode->GetNumVisiblePoints();
+#endif
 					}
 				}
-			}
-
-#if !(UE_BUILD_SHIPPING)
-			// Prepare bounds
-			if (RegisteredProxy.Component->bDrawNodeBounds)
-			{
-				UpdateData.Bounds.Reset(SelectedNodesData[i].Num());
 
-				for (FLidarPointCloudTraversalOctreeNode* Node : SelectedNodesData[i])
+				// Only calculate if needed
+				if (bCalculateVirtualDepth && RegisteredProxy.ComponentRenderParams.ScalingMethod == ELidarPointCloudScalingMethod::PerPoint)
 				{
-					FVector Extent = RegisteredProxy.TraversalOctree->Extents[Node->Depth];
-					UpdateData.Bounds.Emplace(Node->Center - Extent, Node->Center + Extent);
+					RegisteredProxy.TraversalOctree->CalculateVisibilityStructure(UpdateData.TreeStructure);
 				}
-			}
-#endif
 
-			ProxyUpdateData.Add(UpdateData);
+				ProxyUpdateData.Add(UpdateData);
+			}
 		}
 	}
 
-	// Begin streaming data
-	int32 LoadedNodes = 0;
-	for (int32 i = 0; i < InRegisteredProxies.Num(); ++i)
+	// Set when to release the BulkData, if no longer visible
+	const float BulkDataLifetime = CurrentTime + GetDefault<ULidarPointCloudSettings>()->CachedNodeLifetime;
+
+	// Perform data streaming in a separate thread
+	Async(EAsyncExecution::ThreadPool, [OctreeStreamingMap = MoveTemp(OctreeStreamingMap), CurrentTime, BulkDataLifetime]() mutable
 	{
-		const FLidarPointCloudLODManager::FRegisteredProxy& RegisteredProxy = InRegisteredProxies[i];
+		SCOPE_CYCLE_COUNTER(STAT_NodeStreaming);
 
-		FScopeLock OctreeLock(&RegisteredProxy.PointCloud->Octree.DataLock);
-		RegisteredProxy.PointCloud->Octree.UnloadOldNodes(CurrentTime);
-		RegisteredProxy.PointCloud->Octree.StreamQueuedNodes();
+		int32 LoadedNodes = 0;
+		for (TPair<FLidarPointCloudOctree*, TArray<FLidarPointCloudOctreeNode*>>& OctreeStreamingData : OctreeStreamingMap)
+		{
+			OctreeStreamingData.Key->StreamNodes(OctreeStreamingData.Value, CurrentTime);
+			LoadedNodes += OctreeStreamingData.Key->GetNumNodesInUse();
+		}
 
-		LoadedNodes += RegisteredProxy.PointCloud->Octree.GetNumNodesInUse();
-	}
-	SET_DWORD_STAT(STAT_LoadedNodes, LoadedNodes);
+		SET_DWORD_STAT(STAT_LoadedNodes, LoadedNodes);
+	});
 
 	// Update Render Data
 	if (TotalPointsSelected > 0)
@@ -660,7 +688,7 @@ int64 FLidarPointCloudLODManager::ProcessLOD(const TArray<FLidarPointCloudLODMan
 					{
 						for (const FLidarPointCloudProxyUpdateDataNode& Node : UpdateData.SelectedNodes)
 						{
-							if (Node.DataNode->BuildDataCache())
+							if (Node.DataNode->BuildDataCache(UpdateData.bUseStaticBuffers))
 							{
 								MaxPointsPerNode = FMath::Max(MaxPointsPerNode, Node.DataNode->GetNumVisiblePoints());
 							}
@@ -684,8 +712,6 @@ int64 FLidarPointCloudLODManager::ProcessLOD(const TArray<FLidarPointCloudLODMan
 		});
 	}
 
-	bProcessing = false;
-
 	return NewNumPointsInFrustum;
 }
 
@@ -707,73 +733,69 @@ void FLidarPointCloudLODManager::PrepareProxies()
 		// Acquire a Shared Pointer from the Weak Pointer and check that it references a valid object
 		if (TSharedPtr<FLidarPointCloudSceneProxyWrapper, ESPMode::ThreadSafe> SceneProxyWrapper = RegisteredProxy.SceneProxyWrapper.Pin())
 		{
-			ULidarPointCloudComponent* Component = nullptr;
-			if (TWeakObjectPtr<ULidarPointCloudComponent>* PtrToWeakComponent = ProxyComponentMap.Find(RegisteredProxy.SceneProxyWrapper))
+			if (ULidarPointCloudComponent* Component = RegisteredProxy.Component.Get())
 			{
-				Component = (*PtrToWeakComponent).Get();
-
-				if (!Component)
-				{
-					ProxyComponentMap.Remove(SceneProxyWrapper);
-				}
-			}
-
-			if (Component)
-			{
-				if (ULidarPointCloud* PointCloud = Component->GetPointCloud())
+				if (ULidarPointCloud* PointCloud = RegisteredProxy.PointCloud.Get())
 				{
+					// Just in case
+					if (Component->GetPointCloud() == PointCloud)
+					{
 #if WITH_EDITOR
-					// Avoid doubling the point allocation of the same asset (once in Editor world and once in PIE world)
-					RegisteredProxy.bSkip = ViewData.bPIE && Component->GetWorld() && Component->GetWorld()->WorldType == EWorldType::Type::Editor;
+						// Avoid doubling the point allocation of the same asset (once in Editor world and once in PIE world)
+						RegisteredProxy.bSkip = ViewData.bPIE && Component->GetWorld() && Component->GetWorld()->WorldType == EWorldType::Type::Editor;
 #endif
 
-					// Check if the component's transform has changed, and invalidate the Traversal Octree if so
-					const FTransform Transform = Component->GetComponentTransform();
-					if (!RegisteredProxy.LastComponentTransform.Equals(Transform))
-					{
-						RegisteredProxy.TraversalOctree->bValid = false;
-						RegisteredProxy.LastComponentTransform = Transform;
-					}
+						// Check if the component's transform has changed, and invalidate the Traversal Octree if so
+						const FTransform Transform = Component->GetComponentTransform();
+						if (!RegisteredProxy.LastComponentTransform.Equals(Transform))
+						{
+							RegisteredProxy.TraversalOctree->bValid = false;
+							RegisteredProxy.LastComponentTransform = Transform;
+						}
 
-					// Re-initialize the traversal octree, if needed
-					if (!RegisteredProxy.TraversalOctree->bValid)
-					{
-						// Update asset reference
-						RegisteredProxy.PointCloud = PointCloud;
+						// Re-initialize the traversal octree, if needed
+						if (!RegisteredProxy.TraversalOctree->bValid)
+						{
+							// Update asset reference
+							RegisteredProxy.PointCloud = PointCloud;
 
-						// Recreate the Traversal Octree
-						RegisteredProxy.TraversalOctree = MakeShareable(new FLidarPointCloudTraversalOctree(&PointCloud->Octree, Component->GetComponentTransform()));
-						RegisteredProxy.PointCloud->Octree.RegisterTraversalOctree(RegisteredProxy.TraversalOctree);
-					}
+							// Recreate the Traversal Octree
+							RegisteredProxy.TraversalOctree = MakeShareable(new FLidarPointCloudTraversalOctree(&PointCloud->Octree, Component->GetComponentTransform()));
+							RegisteredProxy.PointCloud->Octree.RegisterTraversalOctree(RegisteredProxy.TraversalOctree);
+						}
 
-					// If this is an editor component, use its own ViewportClient
-					if (TSharedPtr<FViewportClient> Client = Component->GetOwningViewportClient().Pin())
-					{
-						// If the ViewData cannot be successfully retrieved from the editor viewport, fall back to using main view
-						if (!RegisteredProxy.ViewData.ComputeFromEditorViewportClient(Client.Get()))
+						// If this is an editor component, use its own ViewportClient
+						if (TSharedPtr<FViewportClient> Client = Component->GetOwningViewportClient().Pin())
+						{
+							// If the ViewData cannot be successfully retrieved from the editor viewport, fall back to using main view
+							if (!RegisteredProxy.ViewData.ComputeFromEditorViewportClient(Client.Get()))
+							{
+								RegisteredProxy.ViewData = ViewData;
+							}
+						}
+						// ... otherwise, use the ViewData provided
+						else
 						{
 							RegisteredProxy.ViewData = ViewData;
 						}
-					}
-					// ... otherwise, use the ViewData provided
-					else
-					{
-						RegisteredProxy.ViewData = ViewData;
-					}
 
-					// Increase priority, if the viewport has focus
-					if (bPrioritizeActiveViewport && RegisteredProxy.ViewData.bHasFocus)
-					{
-						RegisteredProxy.ViewData.ScreenSizeFactor *= 6;
-					}
+						// Increase priority, if the viewport has focus
+						if (bPrioritizeActiveViewport && RegisteredProxy.ViewData.bHasFocus)
+						{
+							RegisteredProxy.ViewData.ScreenSizeFactor *= 6;
+						}
 
-					// Don't count the skippable proxies
-					if (!RegisteredProxy.bSkip)
-					{
-						TotalPointCount += PointCloud->GetNumPoints();
-					}
+						// Don't count the skippable proxies
+						if (!RegisteredProxy.bSkip)
+						{
+							TotalPointCount += PointCloud->GetNumPoints();
+						}
+
+						// Update render params
+						RegisteredProxy.ComponentRenderParams.UpdateFromComponent(Component);
 
-					bValidProxy = true;
+						bValidProxy = true;
+					}
 				}
 			}
 		}
@@ -788,14 +810,14 @@ void FLidarPointCloudLODManager::PrepareProxies()
 	SET_DWORD_STAT(STAT_PointCountTotal, TotalPointCount / 1000);
 }
 
-TArray<const ALidarClippingVolume*> FLidarPointCloudLODManager::GetClippingVolumes() const
+TArray<FLidarPointCloudClippingVolumeParams> FLidarPointCloudLODManager::GetClippingVolumes() const
 {
-	TArray<const ALidarClippingVolume*> ClippingVolumes;
+	TArray<FLidarPointCloudClippingVolumeParams> ClippingVolumes;
 	TArray<UWorld*> Worlds;
 
 	for (int32 i = 0; i < RegisteredProxies.Num(); ++i)
 	{
-		if (ULidarPointCloudComponent* Component = RegisteredProxies[i].Component)
+		if (ULidarPointCloudComponent* Component = RegisteredProxies[i].Component.Get())
 		{
 			if (!Component->IsOwnedByEditor())
 			{
@@ -814,21 +836,24 @@ TArray<const ALidarClippingVolume*> FLidarPointCloudLODManager::GetClippingVolum
 			ALidarClippingVolume* Volume = *It;
 			if (Volume->bEnabled)
 			{
-				ClippingVolumes.Add(Volume);
+				ClippingVolumes.Emplace(Volume);
 			}
 		}
 	}
 
+	ClippingVolumes.Sort();
+
 	return ClippingVolumes;
 }
 
-FLidarPointCloudLODManager::FRegisteredProxy::FRegisteredProxy(ULidarPointCloudComponent* Component, TWeakPtr<FLidarPointCloudSceneProxyWrapper, ESPMode::ThreadSafe> SceneProxyWrapper)
+FLidarPointCloudLODManager::FRegisteredProxy::FRegisteredProxy(TWeakObjectPtr<ULidarPointCloudComponent> Component, TWeakPtr<FLidarPointCloudSceneProxyWrapper, ESPMode::ThreadSafe> SceneProxyWrapper)
 	: Component(Component)
 	, PointCloud(Component->GetPointCloud())
+	, Octree(&PointCloud->Octree)
 	, SceneProxyWrapper(SceneProxyWrapper)
-	, TraversalOctree(new FLidarPointCloudTraversalOctree(&PointCloud->Octree, Component->GetComponentTransform()))
+	, TraversalOctree(new FLidarPointCloudTraversalOctree(Octree, Component->GetComponentTransform()))
 	, LastComponentTransform(Component->GetComponentTransform())
 	, bSkip(false)
 {
-	PointCloud->Octree.RegisterTraversalOctree(TraversalOctree);
+	Octree->RegisterTraversalOctree(TraversalOctree);
 }
diff --git a/Engine/Plugins/Enterprise/LidarPointCloud/Source/LidarPointCloudRuntime/Private/LidarPointCloudOctree.cpp b/Engine/Plugins/Enterprise/LidarPointCloud/Source/LidarPointCloudRuntime/Private/LidarPointCloudOctree.cpp
--- a/Engine/Plugins/Enterprise/LidarPointCloud/Source/LidarPointCloudRuntime/Private/LidarPointCloudOctree.cpp
+++ b/Engine/Plugins/Enterprise/LidarPointCloud/Source/LidarPointCloudRuntime/Private/LidarPointCloudOctree.cpp
@@ -11,8 +11,6 @@
 #include "Misc/FileHelper.h"
 #include "Rendering/LidarPointCloudRenderBuffers.h"
 
-DECLARE_CYCLE_STAT(TEXT("Node Streaming"), STAT_NodeStreaming, STATGROUP_LidarPointCloud);
-
 int32 FLidarPointCloudOctree::MaxNodeDepth = (1 << (sizeof(FLidarPointCloudOctreeNode::Depth) * 8)) - 1;
 int32 FLidarPointCloudOctree::MaxBucketSize = 200;
 int32 FLidarPointCloudOctree::NodeGridResolution = 96;
@@ -125,6 +123,7 @@ FLidarPointCloudOctreeNode::FLidarPointCloudOctreeNode(FLidarPointCloudOctree* T
 	, NumPoints(0)
 	, bHasData(false)
 	, DataCache(nullptr)
+	, VertexFactory(nullptr)
 	, bRenderDataDirty(true)
 	, bHasDataPending(false)
 	, bCanReleaseData(true)
@@ -143,17 +142,7 @@ FLidarPointCloudOctreeNode::~FLidarPointCloudOctreeNode()
 		Children[i] = nullptr;
 	}
 
-	if (DataCache)
-	{
-		FLidarPointCloudRenderBuffer* Tmp = DataCache;
-		ENQUEUE_RENDER_COMMAND(LidarPointCloudOctreeNode_ReleaseDataCache)([Tmp](FRHICommandListImmediate& RHICmdList)
-			{
-				Tmp->ReleaseResource();
-				delete Tmp;
-			});
-
-		DataCache = nullptr;
-	}
+	ReleaseDataCache();
 }
 
 void FLidarPointCloudOctreeNode::UpdateNumVisiblePoints()
@@ -198,28 +187,64 @@ FLidarPointCloudPoint* FLidarPointCloudOctreeNode::GetPersistentData() const
 	return GetData();
 }
 
-bool FLidarPointCloudOctreeNode::BuildDataCache()
+bool FLidarPointCloudOctreeNode::BuildDataCache(bool bUseStaticBuffers)
 {
-	// Only inlcude nodes with available data
+	// Only include nodes with available data
 	if (HasData() && GetNumVisiblePoints())
 	{
-		if (!DataCache)
+		// Make sure to release the unnecessary buffer
+		if (bUseStaticBuffers)
+		{
+			if (DataCache)
+			{
+				DataCache->ReleaseResource();
+				delete DataCache;
+				DataCache = nullptr;
+				bRenderDataDirty = true;
+			}
+
+			if (!VertexFactory)
+			{
+				VertexFactory = new FLidarPointCloudVertexFactory();
+				bRenderDataDirty = true;
+			}
+		}
+		else
 		{
-			DataCache = new FLidarPointCloudRenderBuffer();
-			bRenderDataDirty = true;
+			if (VertexFactory)
+			{
+				VertexFactory->ReleaseResource();
+				delete VertexFactory;
+				VertexFactory = nullptr;
+				bRenderDataDirty = true;
+			}
+
+			if (!DataCache)
+			{
+				DataCache = new FLidarPointCloudRenderBuffer();
+				bRenderDataDirty = true;
+			}
 		}
 
 		if (bRenderDataDirty)
 		{
-			DataCache->Resize(GetNumVisiblePoints() * 5);
+			if (DataCache)
+			{
+				DataCache->Resize(GetNumVisiblePoints() * 5);
 
-			uint8* StructuredBuffer = (uint8*)RHILockBuffer(DataCache->Buffer, 0, GetNumVisiblePoints() * 20, RLM_WriteOnly);
-			for (FLidarPointCloudPoint* P = GetData(), *DataEnd = P + GetNumVisiblePoints(); P != DataEnd; ++P)
+				uint8* StructuredBuffer = (uint8*)RHILockBuffer(DataCache->Buffer, 0, GetNumVisiblePoints() * sizeof(FLidarPointCloudPoint), RLM_WriteOnly);
+				for (FLidarPointCloudPoint* P = GetData(), *DataEnd = P + GetNumVisiblePoints(); P != DataEnd; ++P)
+				{
+					FMemory::Memcpy(StructuredBuffer, P, sizeof(FLidarPointCloudPoint));
+					StructuredBuffer += sizeof(FLidarPointCloudPoint);
+				}
+				RHIUnlockBuffer(DataCache->Buffer);
+			}
+
+			if (VertexFactory)
 			{
-				FMemory::Memcpy(StructuredBuffer, P, 20);
-				StructuredBuffer += 20;
+				VertexFactory->Initialize(GetData(), GetNumVisiblePoints());
 			}
-			RHIUnlockBuffer(DataCache->Buffer);
 
 			bRenderDataDirty = false;
 		}
@@ -253,7 +278,39 @@ FLidarPointCloudOctreeNode* FLidarPointCloudOctreeNode::GetChildNodeAtLocation(c
 	return nullptr;
 }
 
+uint8 FLidarPointCloudOctreeNode::GetChildrenBitmask() const
+{
+	uint8 Bitmask = 0;
+
+	for (FLidarPointCloudOctreeNode* Child : Children)
+	{
+		Bitmask |= 1 << Child->LocationInParent;
+	}
+
+	return Bitmask;
+}
+
 void FLidarPointCloudOctreeNode::InsertPoints(const FLidarPointCloudPoint* Points, const int64& Count, ELidarPointCloudDuplicateHandling DuplicateHandling, const FVector& Translation)
+{
+	InsertPoints_Internal(Points, Count, DuplicateHandling, Translation);
+}
+
+void FLidarPointCloudOctreeNode::InsertPoints_Dynamic(const FLidarPointCloudPoint* Points, const int64& Count, const FVector& Translation)
+{
+	if (Translation.IsNearlyZero())
+	{
+		Data.Append(Points, Count);
+	}
+	else
+	{
+		for (const FLidarPointCloudPoint* PointsPtr = Points, *DataEnd = PointsPtr + Count; PointsPtr != DataEnd; ++PointsPtr)
+		{
+			Data.Emplace(PointsPtr->Location + Translation, PointsPtr->Color, !!PointsPtr->bVisible, PointsPtr->ClassificationID, PointsPtr->Normal);
+		}
+	}
+}
+
+void FLidarPointCloudOctreeNode::InsertPoints_Static(const FLidarPointCloudPoint* Points, const int64& Count, ELidarPointCloudDuplicateHandling DuplicateHandling, const FVector& Translation)
 {
 	const FLidarPointCloudOctree::FSharedLODData& LODData = Tree->SharedData[Depth];
 
@@ -438,12 +495,36 @@ void FLidarPointCloudOctreeNode::InsertPoints(const FLidarPointCloudPoint* Point
 	{
 		if (PointBuckets[i].Num() > 0)
 		{
-			GetChildNodeAtLocation(i)->InsertPoints(PointBuckets[i].GetData(), PointBuckets[i].Num(), DuplicateHandling, FVector::ZeroVector);
+			GetChildNodeAtLocation(i)->InsertPoints_Static(PointBuckets[i].GetData(), PointBuckets[i].Num(), DuplicateHandling, FVector::ZeroVector);
 		}
 	}
 }
 
 void FLidarPointCloudOctreeNode::InsertPoints(FLidarPointCloudPoint** Points, const int64& Count, ELidarPointCloudDuplicateHandling DuplicateHandling, const FVector& Translation)
+{
+	InsertPoints_Internal(Points, Count, DuplicateHandling, Translation);
+}
+
+void FLidarPointCloudOctreeNode::InsertPoints_Dynamic(FLidarPointCloudPoint** Points, const int64& Count, const FVector& Translation)
+{
+	if (Translation.IsNearlyZero())
+	{
+		for (FLidarPointCloudPoint** PointsPtr = Points, **DataEnd = PointsPtr + Count; PointsPtr != DataEnd; ++PointsPtr)
+		{
+			Data.Add(**PointsPtr);
+		}
+	}
+	else
+	{
+		for (FLidarPointCloudPoint** PointsPtr = Points, **DataEnd = PointsPtr + Count; PointsPtr != DataEnd; ++PointsPtr)
+		{
+			const FLidarPointCloudPoint* PointData = *PointsPtr;
+			Data.Emplace(PointData->Location + Translation, PointData->Color, !!PointData->bVisible, PointData->ClassificationID, PointData->Normal);
+		}
+	}
+}
+
+void FLidarPointCloudOctreeNode::InsertPoints_Static(FLidarPointCloudPoint** Points, const int64& Count, ELidarPointCloudDuplicateHandling DuplicateHandling, const FVector& Translation)
 {
 	const FLidarPointCloudOctree::FSharedLODData& LODData = Tree->SharedData[Depth];
 
@@ -628,11 +709,33 @@ void FLidarPointCloudOctreeNode::InsertPoints(FLidarPointCloudPoint** Points, co
 	{
 		if (PointBuckets[i].Num() > 0)
 		{
-			GetChildNodeAtLocation(i)->InsertPoints(PointBuckets[i].GetData(), PointBuckets[i].Num(), DuplicateHandling, FVector::ZeroVector);
+			GetChildNodeAtLocation(i)->InsertPoints_Static(PointBuckets[i].GetData(), PointBuckets[i].Num(), DuplicateHandling, FVector::ZeroVector);
 		}
 	}
 }
 
+template <typename T>
+void FLidarPointCloudOctreeNode::InsertPoints_Internal(T Points, const int64& Count, ELidarPointCloudDuplicateHandling DuplicateHandling, const FVector& Translation)
+{
+	if (Tree->IsOptimizedForDynamicData())
+	{
+		Data.Reserve(NumPoints + Count);
+
+		InsertPoints_Dynamic(Points, Count, Translation);
+
+		NumPoints = Data.Num();
+		bCanReleaseData = false;
+		bRenderDataDirty = true;
+		bHasData = true;
+
+		AddPointCount(Count);
+	}
+	else
+	{
+		InsertPoints_Static(Points, Count, DuplicateHandling, Translation);
+	}
+}
+
 void FLidarPointCloudOctreeNode::Empty(bool bRecursive)
 {
 	if (bHasData)
@@ -640,6 +743,7 @@ void FLidarPointCloudOctreeNode::Empty(bool bRecursive)
 		bHasData = false;
 		bInUse = false;
 		NumPoints = 0;
+		NumVisiblePoints = 0;
 		Data.Empty();
 	}
 
@@ -704,16 +808,30 @@ void FLidarPointCloudOctreeNode::ReleaseData(bool bForce)
 		}
 	}
 
-	if (DataCache)
+	ReleaseDataCache();
+}
+
+void FLidarPointCloudOctreeNode::ReleaseDataCache()
+{
+	if (VertexFactory || DataCache)
 	{
-		FLidarPointCloudRenderBuffer* Tmp = DataCache;
-		ENQUEUE_RENDER_COMMAND(LidarPointCloudOctreeNode_ReleaseDataCache)([Tmp](FRHICommandListImmediate& RHICmdList)
+		ENQUEUE_RENDER_COMMAND(LidarPointCloudOctreeNode_ReleaseDataCache)([DataCache = DataCache, VertexFactory = VertexFactory](FRHICommandListImmediate& RHICmdList)
 			{
-				Tmp->ReleaseResource();
-				delete Tmp;
+				if (DataCache)
+				{
+					DataCache->ReleaseResource();
+					delete DataCache;
+				}
+
+				if (VertexFactory)
+				{
+					VertexFactory->ReleaseResource();
+					delete VertexFactory;
+				}
 			});
 
 		DataCache = nullptr;
+		VertexFactory = nullptr;
 	}
 }
 
@@ -1007,10 +1125,10 @@ void FLidarPointCloudOctree::GetPointsInBox_Internal(TArray<const FLidarPointClo
 }
 
 template <typename T>
-void FLidarPointCloudOctree::GetPointsInFrustum_Internal(TArray<FLidarPointCloudPoint*, T>& SelectedPoints, const FConvexVolume& Frustum, const bool& bVisibleOnly)
+void FLidarPointCloudOctree::GetPointsInConvexVolume_Internal(TArray<FLidarPointCloudPoint*, T>& SelectedPoints, const FConvexVolume& ConvexVolume, const bool& bVisibleOnly)
 {
 	SelectedPoints.Reset();
-	PROCESS_IN_FRUSTUM({ SelectedPoints.Add(Point); });
+	PROCESS_IN_CONVEX_VOLUME({ SelectedPoints.Add(Point); });
 }
 
 template <typename T>
@@ -1147,14 +1265,24 @@ void FLidarPointCloudOctree::GetPointsInBox(TArray<const FLidarPointCloudPoint*>
 	GetPointsInBox_Internal(SelectedPoints, Box, bVisibleOnly);
 }
 
-void FLidarPointCloudOctree::GetPointsInFrustum(TArray64<FLidarPointCloudPoint*>& SelectedPoints, const FConvexVolume& Frustum, const bool& bVisibleOnly)
+void FLidarPointCloudOctree::GetPointsInConvexVolume(TArray<FLidarPointCloudPoint*>& SelectedPoints, const FConvexVolume& ConvexVolume, const bool& bVisibleOnly)
+{
+	GetPointsInConvexVolume_Internal(SelectedPoints, ConvexVolume, bVisibleOnly);
+}
+
+void FLidarPointCloudOctree::GetPointsInConvexVolume(TArray64<FLidarPointCloudPoint*>& SelectedPoints, const FConvexVolume& ConvexVolume, const bool& bVisibleOnly)
 {
-	GetPointsInFrustum_Internal(SelectedPoints, Frustum, bVisibleOnly);
+	GetPointsInConvexVolume_Internal(SelectedPoints, ConvexVolume, bVisibleOnly);
 }
 
 void FLidarPointCloudOctree::GetPointsInFrustum(TArray<FLidarPointCloudPoint*>& SelectedPoints, const FConvexVolume& Frustum, const bool& bVisibleOnly)
 {
-	GetPointsInFrustum_Internal(SelectedPoints, Frustum, bVisibleOnly);
+	GetPointsInConvexVolume_Internal(SelectedPoints, Frustum, bVisibleOnly);
+}
+
+void FLidarPointCloudOctree::GetPointsInFrustum(TArray64<FLidarPointCloudPoint*>& SelectedPoints, const FConvexVolume& Frustum, const bool& bVisibleOnly)
+{
+	GetPointsInConvexVolume_Internal(SelectedPoints, Frustum, bVisibleOnly);
 }
 
 void FLidarPointCloudOctree::GetPointsAsCopies(TArray64<FLidarPointCloudPoint>& SelectedPoints, const FTransform* LocalToWorld, int64 StartIndex /*= 0*/, int64 Count /*= -1*/) const
@@ -1584,9 +1712,20 @@ void FLidarPointCloudOctree::MarkRenderDataDirty()
 	ITERATE_NODES({ CurrentNode->bRenderDataDirty = true; }, true);
 }
 
+void FLidarPointCloudOctree::MarkRenderDataInSphereDirty(const FSphere& Sphere)
+{
+	const FBox Box(Sphere.Center - FVector(Sphere.W), Sphere.Center + FVector(Sphere.W));
+	ITERATE_NODES({ CurrentNode->bRenderDataDirty = true; }, NODE_IN_BOX);
+}
+
+void FLidarPointCloudOctree::MarkRenderDataInConvexVolumeDirty(const FConvexVolume& ConvexVolume)
+{
+	ITERATE_NODES({ CurrentNode->bRenderDataDirty = true; }, NODE_IN_CONVEX_VOLUME);
+}
+
 void FLidarPointCloudOctree::MarkRenderDataInFrustumDirty(const FConvexVolume& Frustum)
 {
-	ITERATE_NODES({ CurrentNode->bRenderDataDirty = true; }, NODE_IN_FRUSTUM);
+	MarkRenderDataInConvexVolumeDirty(Frustum);
 }
 
 void FLidarPointCloudOctree::Initialize(const FVector& InExtent)
@@ -1914,8 +2053,6 @@ void FLidarPointCloudOctree::Empty(bool bDestroyNodes)
 
 		QueuedNodes.Empty();
 		NodesInUse.Reset();
-
-		MarkTraversalOctreesForInvalidation();
 	}
 	else
 	{
@@ -1927,6 +2064,8 @@ void FLidarPointCloudOctree::Empty(bool bDestroyNodes)
 	{
 		Count.Reset();
 	}
+
+	MarkTraversalOctreesForInvalidation();
 }
 
 void FLidarPointCloudOctree::UnregisterTraversalOctree(FLidarPointCloudTraversalOctree* TraversalOctree)
@@ -1961,86 +2100,72 @@ void FLidarPointCloudOctree::UnregisterTraversalOctree(FLidarPointCloudTraversal
 	}
 }
 
-void FLidarPointCloudOctree::QueueNode(FLidarPointCloudOctreeNode* Node, float Lifetime)
+void FLidarPointCloudOctree::StreamNodes(TArray<FLidarPointCloudOctreeNode*>& Nodes, const float& CurrentTime)
 {
-	if (!Node)
-	{
-		return;
-	}
+	const float BulkDataLifetime = CurrentTime + GetDefault<ULidarPointCloudSettings>()->CachedNodeLifetime;
 
-	// Refresh lifetime of the BulkData, if requested
-	if (Lifetime > -1)
+	for (FLidarPointCloudOctreeNode* Node : Nodes)
 	{
-		Node->BulkDataLifetime = Lifetime;
-	}
+		// Refresh lifetime of the BulkData
+		Node->BulkDataLifetime = BulkDataLifetime;
 
-	// Mark the node as being in use
-	if (!Node->bInUse)
-	{
-		Node->bInUse = true;
-		NodesInUse.Add(Node);
+		// Enqueue for processing
+		QueuedNodes.Enqueue(Node);
 	}
 
-	// No need to do anything, if the node already has data loaded or loading
-	if (Node->HasData() || Node->bHasDataPending)
+	// Only one process at a time
+	if (!bStreamingBusy)
 	{
-		return;
-	}
+		bStreamingBusy = true;
 
-	QueuedNodes.Enqueue(Node);
-	Node->bHasDataPending = true;
-}
-
-void FLidarPointCloudOctree::StreamQueuedNodes()
-{
-	SCOPE_CYCLE_COUNTER(STAT_NodeStreaming);
-
-	// Only one streaming operation at a time
-	if (bStreamingBusy)
-	{
-		return;
-	}
-
-	bStreamingBusy = true;
-	
-	// Perform data streaming in a separate thread
-	Async(EAsyncExecution::ThreadPool, [this]
-	{
-		SCOPE_CYCLE_COUNTER(STAT_NodeStreaming);
-
-		FLidarPointCloudOctreeNode* CurrentNode = nullptr;
-		while (QueuedNodes.Dequeue(CurrentNode))
+		// Add previously queued nodes
+		FLidarPointCloudOctreeNode* QueuedNode;
+		while (QueuedNodes.Dequeue(QueuedNode))
 		{
-			StreamNodeData(CurrentNode);
-			CurrentNode->bHasDataPending = false;
+			if (!QueuedNode->bInUse)
+			{
+				NodesInUse.Add(QueuedNode);
+			}
 		}
 
-		bStreamingBusy = false;
-	});
-}
-
-void FLidarPointCloudOctree::UnloadOldNodes(const float& CurrentTime)
-{
-	SCOPE_CYCLE_COUNTER(STAT_NodeStreaming);
-
-	// Skip if the data is locked
-	FScopeTryLock ReleaseLock(&DataReleaseLock);
-	if (!ReleaseLock.IsLocked())
-	{
-		return;
-	}
+		// Process nodes' streaming requests
+		// If the release lock has been acquired, release unused nodes in one go
+		// Otherwise, just load the ones requested
+		FScopeTryLock ReleaseLock(&DataReleaseLock);
+		if(ReleaseLock.IsLocked())
+		{
+			for (int32 i = 0; i < NodesInUse.Num(); ++i)
+			{
+				FLidarPointCloudOctreeNode* Node = NodesInUse[i];
 
-	for(int32 i = 0; i < NodesInUse.Num(); ++i)
-	{
-		FLidarPointCloudOctreeNode* Node = NodesInUse[i];
+				// Check if the node should still be in use
+				Node->bInUse = Node->BulkDataLifetime >= CurrentTime;
 
-		// Unload data, if it expired
-		if (Node->BulkDataLifetime < CurrentTime)
+				// If node is alive, make sure it has data loaded
+				if (Node->bInUse)
+				{
+					Node->GetData();
+				}
+				// ... otherwise, release it
+				else
+				{
+					Node->ReleaseData();				
+					NodesInUse.RemoveAtSwap(i--, 1, false);
+				}
+			}
+		}
+		else
 		{
-			Node->ReleaseData();
-			Node->bInUse = false;
-			NodesInUse.RemoveAtSwap(i--, 1, false);
+			for (FLidarPointCloudOctreeNode* Node : NodesInUse)
+			{
+				if (Node->BulkDataLifetime >= CurrentTime)
+				{
+					Node->GetData();
+				}
+			}
 		}
+
+		bStreamingBusy = false;
 	}
 }
 
@@ -2067,6 +2192,52 @@ void FLidarPointCloudOctree::ReleaseAllNodes(bool bIncludePersistent)
 	}
 }
 
+bool FLidarPointCloudOctree::IsOptimizedForDynamicData() const
+{
+	return Owner && Owner->IsOptimizedForDynamicData();
+}
+
+void FLidarPointCloudOctree::OptimizeForDynamicData()
+{
+	if (!IsOptimizedForDynamicData())
+	{
+		// Move all data from the nodes
+		TArray<FLidarPointCloudPoint> SourceData;
+		SourceData.Reserve(GetNumPoints());
+		
+		ITERATE_NODES(
+		{
+			SourceData.Append(CurrentNode->Data);
+			CurrentNode->Data.Empty();
+		}, true);
+
+		// Destroy current tree structure
+		Empty(true);
+
+		bIsFullyLoaded = true;
+
+		// Insert data back into the tree
+		InsertPoints_Internal(SourceData.GetData(), SourceData.Num(), ELidarPointCloudDuplicateHandling::Ignore, false, FVector::ZeroVector);
+	}
+}
+
+void FLidarPointCloudOctree::OptimizeForStaticData()
+{
+	if (IsOptimizedForDynamicData())
+	{
+		// Move data out of the root node
+		TArray<FLidarPointCloudPoint> SourceData = MoveTemp(Root.Data);
+		
+		// Destroy current tree structure
+		Empty(true);
+
+		bIsFullyLoaded = true;
+
+		// Insert data back into the tree
+		InsertPoints_Internal(SourceData.GetData(), SourceData.Num(), ELidarPointCloudDuplicateHandling::Ignore, false, FVector::ZeroVector);
+	}
+}
+
 void FLidarPointCloudOctree::RefreshAllocatedSize()
 {
 	FScopeTryLock Lock(&DataLock);
@@ -2426,8 +2597,9 @@ FLidarPointCloudTraversalOctreeNode::FLidarPointCloudTraversalOctreeNode()
 
 }
 
-void FLidarPointCloudTraversalOctreeNode::Build(FLidarPointCloudOctreeNode* Node, const FTransform& LocalToWorld, const FVector& LocationOffset)
+void FLidarPointCloudTraversalOctreeNode::Build(FLidarPointCloudTraversalOctree* TraversalOctree, FLidarPointCloudOctreeNode* Node, const FTransform& LocalToWorld, const FVector& LocationOffset)
 {
+	Octree = TraversalOctree;
 	DataNode = Node;
 	Center = LocalToWorld.TransformPosition(Node->Center + LocationOffset);
 	Depth = Node->Depth;
@@ -2437,19 +2609,22 @@ void FLidarPointCloudTraversalOctreeNode::Build(FLidarPointCloudOctreeNode* Node
 	{
 		if (Node->Children[i])
 		{
-			Children[i].Build(Node->Children[i], LocalToWorld, LocationOffset);
+			Children[i].Build(Octree, Node->Children[i], LocalToWorld, LocationOffset);
 			Children[i].Parent = this;
 		}
 	}
 }
 
-uint8 FLidarPointCloudTraversalOctreeNode::CalculateVirtualDepth(const TArray<float>& LevelWeights, const float& VDMultiplier, const float& PointSizeBias) const
+void FLidarPointCloudTraversalOctreeNode::CalculateVirtualDepth(const TArray<float>& LevelWeights, const float& PointSizeBias)
 {
 	if (!IsAvailable())
 	{
-		return 0;
+		VirtualDepth = 255;
+		return;
 	}
 
+	const float& VDMultiplier = Octree->VirtualDepthMultiplier;
+
 	TQueue<const FLidarPointCloudTraversalOctreeNode*> Nodes;
 	const FLidarPointCloudTraversalOctreeNode* CurrentNode = nullptr;
 
@@ -2466,7 +2641,7 @@ uint8 FLidarPointCloudTraversalOctreeNode::CalculateVirtualDepth(const TArray<fl
 			}
 		}
 
-		float LocalVDFactor = CurrentNode->Depth * CurrentNode->DataNode->GetNumPoints() * LevelWeights[CurrentNode->Depth];
+		float LocalVDFactor = CurrentNode->Depth * CurrentNode->DataNode->GetNumVisiblePoints() * LevelWeights[CurrentNode->Depth];
 
 		if (CurrentNode != this && PointSizeBias > 0)
 		{
@@ -2489,11 +2664,20 @@ uint8 FLidarPointCloudTraversalOctreeNode::CalculateVirtualDepth(const TArray<fl
 			}
 		}
 
-		NumPoints += CurrentNode->DataNode->GetNumPoints() * LevelWeights[CurrentNode->Depth];
+		NumPoints += CurrentNode->DataNode->GetNumVisiblePoints() * LevelWeights[CurrentNode->Depth];
 	}
 
 	// Calculate the Virtual Depth
-	return VDFactor / NumPoints * VDMultiplier;
+	VirtualDepth = VDFactor / NumPoints * VDMultiplier;
+}
+
+//////////////////////////////////////////////////////////// FLidarPointCloudTraversalOctreeNodeSizeData
+
+FLidarPointCloudTraversalOctreeNodeSizeData::FLidarPointCloudTraversalOctreeNodeSizeData(FLidarPointCloudTraversalOctreeNode* Node, const float& Size, const int32& ProxyIndex)
+	: Node(Node)
+	, Size(Size)
+	, ProxyIndex(ProxyIndex)
+{
 }
 
 //////////////////////////////////////////////////////////// FLidarPointCloudTraversalOctree
@@ -2538,11 +2722,73 @@ FLidarPointCloudTraversalOctree::FLidarPointCloudTraversalOctree(FLidarPointClou
 	}
 
 	// Star cloning the node data
-	Root.Build(&Octree->Root, LocalToWorld, Octree->Owner->GetLocationOffset().ToVector());
+	Root.Build(this, &Octree->Root, LocalToWorld, Octree->Owner->GetLocationOffset().ToVector());
 
 	bValid = true;
 }
 
+void FLidarPointCloudTraversalOctree::CalculateVisibilityStructure(TArray<uint32>& OutData)
+{
+	FLidarPointCloudTraversalOctreeNode* CurrentNode = nullptr;
+	TQueue<FLidarPointCloudTraversalOctreeNode*> Nodes;
+	Nodes.Enqueue(&Root);
+
+	uint16 NumNodesInQueue = 1;
+
+	while (Nodes.Dequeue(CurrentNode))
+	{
+		uint32& Data = OutData[OutData.AddZeroed()];
+		Data |= (0x00FFFF00 & (NumNodesInQueue-- << 8));
+		Data |= 0xFF000000 & (CurrentNode->VirtualDepth << 24);
+
+		for (uint8 i = 0; i < 8; ++i)
+		{
+			for (FLidarPointCloudTraversalOctreeNode& Child : CurrentNode->Children)
+			{
+				if (Child.bSelected && Child.DataNode->LocationInParent == i)
+				{
+					Nodes.Enqueue(&Child);
+					Data |= 1 << i;
+					++NumNodesInQueue;
+				}
+			}
+		}
+	}
+}
+
+void FLidarPointCloudTraversalOctree::CalculateLevelWeightsForSelectedNodes(TArray<float>& OutLevelWeights)
+{
+	FLidarPointCloudTraversalOctreeNode* CurrentNode = nullptr;
+	TQueue<FLidarPointCloudTraversalOctreeNode*> Nodes;
+	Nodes.Enqueue(&Root);
+
+	OutLevelWeights.Empty();
+	OutLevelWeights.AddZeroed(NumLODs);
+	TArray<int64> PointCount;
+	PointCount.AddZeroed(NumLODs);
+	int64 NumPoints = 0;
+
+	while (Nodes.Dequeue(CurrentNode))
+	{
+		for (FLidarPointCloudTraversalOctreeNode& Child : CurrentNode->Children)
+		{
+			if (Child.bSelected)
+			{
+				Nodes.Enqueue(&Child);
+			}
+		}
+
+		const int32 NumPointsInNode = CurrentNode->DataNode->GetNumVisiblePoints();
+		PointCount[CurrentNode->Depth] += NumPointsInNode;
+		NumPoints += NumPointsInNode;
+	}
+
+	for (int32 i = 0; i < OutLevelWeights.Num(); i++)
+	{
+		OutLevelWeights[i] = NumPoints > 0 ? (float)PointCount[i] / NumPoints : 0;
+	}
+}
+
 FLidarPointCloudTraversalOctree::~FLidarPointCloudTraversalOctree()
 {
 	if (bValid)
diff --git a/Engine/Plugins/Enterprise/LidarPointCloud/Source/LidarPointCloudRuntime/Private/LidarPointCloudOctreeMacros.h b/Engine/Plugins/Enterprise/LidarPointCloud/Source/LidarPointCloudRuntime/Private/LidarPointCloudOctreeMacros.h
--- a/Engine/Plugins/Enterprise/LidarPointCloud/Source/LidarPointCloudRuntime/Private/LidarPointCloudOctreeMacros.h
+++ b/Engine/Plugins/Enterprise/LidarPointCloud/Source/LidarPointCloudRuntime/Private/LidarPointCloudOctreeMacros.h
@@ -8,7 +8,7 @@
 #define IS_VIS_CHECK_REQUIRED (bVisibleOnly && CurrentNode->NumVisiblePoints < CurrentNode->GetNumPoints())
 
 #define NODE_IN_BOX (Box.Intersect(Child->GetBounds()))
-#define NODE_IN_FRUSTUM (Frustum.IntersectBox(Child->Center, SharedData[Child->Depth].Extent))
+#define NODE_IN_CONVEX_VOLUME (ConvexVolume.IntersectBox(Child->Center, SharedData[Child->Depth].Extent))
 
 #define ITERATE_NODES_BODY(Action, NodeTest, Const) \
 {\
@@ -30,7 +30,7 @@
 
 #define POINT_IN_BOX Box.IsInsideOrOn(Point->Location)
 #define POINT_IN_SPHERE (POINT_IN_BOX && FVector::DistSquared(Point->Location, Sphere.Center) <= RadiusSq)
-#define POINT_IN_FRUSTUM Frustum.IntersectSphere(Point->Location, 0)
+#define POINT_IN_CONVEX_VOLUME ConvexVolume.IntersectSphere(Point->Location, 0)
 #define POINT_BY_RAY Ray.Intersects(Point, RadiusSq)
 
 #define PROCESS_BODY(Action, PointTest, Mode) \
@@ -77,13 +77,13 @@
 }
 #define PROCESS_IN_BOX_BODY(Action, Mode) PROCESS_IN_BOX_BODY_EXTERN(this, Action, Mode)
 
-#define PROCESS_IN_FRUSTUM_BODY(Action, Mode) \
+#define PROCESS_IN_CONVEX_VOLUME_BODY(Action, Mode) \
 {\
 	if (!bVisibleOnly || CurrentNode->NumVisiblePoints > 0)\
 	{\
 		bool bNodeFullyContained;\
-		Frustum.IntersectBox(CurrentNode->Center, SharedData[CurrentNode->Depth].Extent, bNodeFullyContained);\
-		PROCESS_BODY(Action, POINT_IN_FRUSTUM, Mode) \
+		ConvexVolume.IntersectBox(CurrentNode->Center, SharedData[CurrentNode->Depth].Extent, bNodeFullyContained);\
+		PROCESS_BODY(Action, POINT_IN_CONVEX_VOLUME, Mode) \
 	}\
 }
 
@@ -126,8 +126,8 @@
 #define PROCESS_IN_BOX_CONST(Action) { ITERATE_NODES_CONST(PROCESS_IN_BOX_BODY(Action, _RO), NODE_IN_BOX) }
 #define PROCESS_IN_BOX_EXTERN(Octree, Action) { ITERATE_NODES(PROCESS_IN_BOX_BODY_EXTERN(Octree, Action,), NODE_IN_BOX) }
 
-#define PROCESS_IN_FRUSTUM(Action) { ITERATE_NODES(PROCESS_IN_FRUSTUM_BODY(Action,), NODE_IN_FRUSTUM) }
-#define PROCESS_IN_FRUSTUM_CONST(Action) { ITERATE_NODES_CONST(PROCESS_IN_FRUSTUM_BODY(Action, _RO), NODE_IN_FRUSTUM) }
+#define PROCESS_IN_CONVEX_VOLUME(Action) { ITERATE_NODES(PROCESS_IN_CONVEX_VOLUME_BODY(Action,), NODE_IN_CONVEX_VOLUME) }
+#define PROCESS_IN_CONVEX_VOLUME_CONST(Action) { ITERATE_NODES_CONST(PROCESS_IN_CONVEX_VOLUME_BODY(Action, _RO), NODE_IN_CONVEX_VOLUME) }
 
 #define PROCESS_BY_RAY_COMMON(Action)\
 {\
diff --git a/Engine/Plugins/Enterprise/LidarPointCloud/Source/LidarPointCloudRuntime/Private/LidarPointCloudSettings.cpp b/Engine/Plugins/Enterprise/LidarPointCloud/Source/LidarPointCloudRuntime/Private/LidarPointCloudSettings.cpp
--- a/Engine/Plugins/Enterprise/LidarPointCloud/Source/LidarPointCloudRuntime/Private/LidarPointCloudSettings.cpp
+++ b/Engine/Plugins/Enterprise/LidarPointCloud/Source/LidarPointCloudRuntime/Private/LidarPointCloudSettings.cpp
@@ -16,6 +16,7 @@ ULidarPointCloudSettings::ULidarPointCloudSettings()
 	, bReleaseAssetAfterCooking(true)
 	, bUseRenderDataSmoothing(true)
 	, RenderDataSmoothingMaxFrametime(0.01f)
+	, bUseFastRendering(true)
 	, MeshingBatchSize(128)
 	, bAutoCenterOnImport(true)
 	, bAutoCalculateNormalsOnImport(false)
diff --git a/Engine/Plugins/Enterprise/LidarPointCloud/Source/LidarPointCloudRuntime/Private/LidarPointCloudShared.cpp b/Engine/Plugins/Enterprise/LidarPointCloud/Source/LidarPointCloudRuntime/Private/LidarPointCloudShared.cpp
--- a/Engine/Plugins/Enterprise/LidarPointCloud/Source/LidarPointCloudRuntime/Private/LidarPointCloudShared.cpp
+++ b/Engine/Plugins/Enterprise/LidarPointCloud/Source/LidarPointCloudRuntime/Private/LidarPointCloudShared.cpp
@@ -132,3 +132,15 @@ void FLidarPointCloudDataBufferManager::Resize(const int32& NewBufferSize)
 		Iterator = Iterator->Next;
 	}
 }
+
+FLidarPointCloudClippingVolumeParams::FLidarPointCloudClippingVolumeParams(const ALidarClippingVolume* ClippingVolume)
+	: Mode(ClippingVolume->Mode)
+	, Priority(ClippingVolume->Priority)
+	, Bounds(ClippingVolume->GetBounds().GetBox())
+{
+	const FVector Extent = ClippingVolume->GetActorScale3D() * 100;
+	PackedShaderData = FMatrix( FPlane(ClippingVolume->GetActorLocation(), Mode == ELidarClippingVolumeMode::ClipInside),
+								FPlane(ClippingVolume->GetActorForwardVector(), Extent.X),
+								FPlane(ClippingVolume->GetActorRightVector(), Extent.Y),
+								FPlane(ClippingVolume->GetActorUpVector(), Extent.Z));
+}
diff --git a/Engine/Plugins/Enterprise/LidarPointCloud/Source/LidarPointCloudRuntime/Private/Rendering/LidarPointCloudRenderBuffers.cpp b/Engine/Plugins/Enterprise/LidarPointCloud/Source/LidarPointCloudRuntime/Private/Rendering/LidarPointCloudRenderBuffers.cpp
--- a/Engine/Plugins/Enterprise/LidarPointCloud/Source/LidarPointCloudRuntime/Private/Rendering/LidarPointCloudRenderBuffers.cpp
+++ b/Engine/Plugins/Enterprise/LidarPointCloud/Source/LidarPointCloudRuntime/Private/Rendering/LidarPointCloudRenderBuffers.cpp
@@ -20,7 +20,7 @@
 //////////////////////////////////////////////////////////// Base Buffer
 
 TGlobalResource<FLidarPointCloudIndexBuffer> GLidarPointCloudIndexBuffer;
-TGlobalResource<FLidarPointCloudVertexFactory> GLidarPointCloudVertexFactory;
+TGlobalResource<FLidarPointCloudSharedVertexFactory> GLidarPointCloudSharedVertexFactory;
 
 //////////////////////////////////////////////////////////// Index Buffer
 
@@ -40,29 +40,26 @@ void FLidarPointCloudIndexBuffer::Resize(const uint32 & RequestedCapacity)
 void FLidarPointCloudIndexBuffer::InitRHI()
 {
 	FRHIResourceCreateInfo CreateInfo(TEXT("FLidarPointCloudIndexBuffer"));
-	uint32 Size = Capacity * 7 * sizeof(uint32);
+	const uint32 Size = Capacity * 7 * sizeof(uint32);
 	PointOffset = Capacity * 6;
 
 	IndexBufferRHI = RHICreateBuffer(Size, BUF_Dynamic | BUF_IndexBuffer, sizeof(uint32), ERHIAccess::VertexOrIndexBuffer, CreateInfo);
 
 	uint32* Data = (uint32*)RHILockBuffer(IndexBufferRHI, 0, Size, RLM_WriteOnly);
-	for (uint32 i = 0; i < Capacity; i++)
+	for (uint32 i = 0, idx = 0; i < Capacity; i++)
 	{
+		const uint32 v = i * 4;
+
 		// Full quads
-		{
-			uint32 idx = i * 6;
-			uint32 v = i * 4;
-
-			Data[idx] = v;
-			Data[idx + 1] = v + 1;
-			Data[idx + 2] = v + 2;
-			Data[idx + 3] = v;
-			Data[idx + 4] = v + 2;
-			Data[idx + 5] = v + 3;
-		}
+		Data[idx++] = v;
+		Data[idx++] = v + 1;
+		Data[idx++] = v + 2;
+		Data[idx++] = v;
+		Data[idx++] = v + 2;
+		Data[idx++] = v + 3;
 
 		// Points
-		Data[PointOffset + i] = i;
+		Data[PointOffset + i] = v;
 	}
 	RHIUnlockBuffer(IndexBufferRHI);
 }
@@ -114,9 +111,6 @@ void FLidarPointCloudRenderBuffer::ReleaseRHI()
 
 FLidarPointCloudBatchElementUserData::FLidarPointCloudBatchElementUserData()
 	: SelectionColor(FVector::OneVector)
-	, IndexDivisor(4)
-	, SpriteSize(0)
-	, bUseLODColoration(false)
 	, NumClippingVolumes(0)
 	, bStartClipped(false)
 {
@@ -146,19 +140,22 @@ void FLidarPointCloudBatchElementUserData::SetClassificationColors(const TMap<in
 
 void FLidarPointCloudVertexFactoryShaderParameters::Bind(const FShaderParameterMap& ParameterMap)
 {
+	BINDPARAM(TreeBuffer);
 	BINDPARAM(DataBuffer);
 	BINDPARAM(bEditorView);
 	BINDPARAM(SelectionColor);
-	BINDPARAM(IndexDivisor);
 	BINDPARAM(LocationOffset);
+	BINDPARAM(RootCellSize);
+	BINDPARAM(RootExtent);
+	BINDPARAM(bUsePerPointScaling);
 	BINDPARAM(VirtualDepth);
-	BINDPARAM(SpriteSize);
-	BINDPARAM(bUseLODColoration);
 	BINDPARAM(SpriteSizeMultiplier);
+	BINDPARAM(ReversedVirtualDepthMultiplier);
 	BINDPARAM(ViewRightVector);
 	BINDPARAM(ViewUpVector);
 	BINDPARAM(bUseCameraFacing);
 	BINDPARAM(bUseScreenSizeScaling);
+	BINDPARAM(bUseStaticBuffers);
 	BINDPARAM(BoundsSize);
 	BINDPARAM(ElevationColorBottom);
 	BINDPARAM(ElevationColorTop);
@@ -183,19 +180,22 @@ void FLidarPointCloudVertexFactoryShaderParameters::GetElementShaderBindings(con
 {
 	FLidarPointCloudBatchElementUserData* UserData = (FLidarPointCloudBatchElementUserData*)BatchElement.UserData;
 
+	SETSRVPARAM(TreeBuffer);
 	SETSRVPARAM(DataBuffer);
 	SETPARAM(bEditorView);
 	SETPARAM(SelectionColor);
-	SETPARAM(IndexDivisor);
 	SETPARAM(LocationOffset);
+	SETPARAM(RootCellSize);
+	SETPARAM(RootExtent);
+	SETPARAM(bUsePerPointScaling);
 	SETPARAM(VirtualDepth);
-	SETPARAM(SpriteSize);
-	SETPARAM(bUseLODColoration);
 	SETPARAM(SpriteSizeMultiplier);
+	SETPARAM(ReversedVirtualDepthMultiplier);
 	SETPARAM(ViewRightVector);
 	SETPARAM(ViewUpVector);
 	SETPARAM(bUseCameraFacing);
 	SETPARAM(bUseScreenSizeScaling);
+	SETPARAM(bUseStaticBuffers);
 	SETPARAM(BoundsSize);
 	SETPARAM(ElevationColorBottom);
 	SETPARAM(ElevationColorTop);
@@ -215,18 +215,52 @@ void FLidarPointCloudVertexFactoryShaderParameters::GetElementShaderBindings(con
 	SETPARAM(bStartClipped);
 }
 
-bool FLidarPointCloudVertexFactory::ShouldCompilePermutation(const FVertexFactoryShaderPermutationParameters& Parameters)
+bool FLidarPointCloudVertexFactoryBase::ShouldCompilePermutation(const FVertexFactoryShaderPermutationParameters& Parameters)
 {
 	return (IsPCPlatform(Parameters.Platform) && IsFeatureLevelSupported(Parameters.Platform, ERHIFeatureLevel::SM5) &&
 		Parameters.MaterialParameters.MaterialDomain == MD_Surface && Parameters.MaterialParameters.bIsUsedWithLidarPointCloud) || Parameters.MaterialParameters.bIsSpecialEngineMaterial;
 }
 
+void FLidarPointCloudVertexFactory::Initialize(FLidarPointCloudPoint* Data, int32 NumPoints)
+{
+	if (IsInitialized())
+	{
+		ReleaseResource();
+	}
+
+	VertexBuffer.Data = Data;
+	VertexBuffer.NumPoints = NumPoints;
+
+	InitResource();
+}
+
+void FLidarPointCloudVertexFactory::FPointCloudVertexBuffer::InitRHI()
+{
+	const uint32 BufferSize = NumPoints * 4 * sizeof(FLidarPointCloudPoint);
+	FRHIResourceCreateInfo CreateInfo(TEXT("FPointCloudVertexBuffer"));
+	VertexBufferRHI = RHICreateBuffer(BufferSize, BUF_Static | BUF_VertexBuffer, sizeof(FLidarPointCloudPoint), ERHIAccess::VertexOrIndexBuffer, CreateInfo);
+	void* Buffer = RHILockBuffer(VertexBufferRHI, 0, BufferSize, RLM_WriteOnly);
+
+	uint8* Dest = (uint8*)Buffer;
+	for (int32 i = 0; i < NumPoints; ++i, ++Data)
+	{
+		FMemory::Memcpy(Dest, Data, sizeof(FLidarPointCloudPoint)); Dest += sizeof(FLidarPointCloudPoint);
+		FMemory::Memcpy(Dest, Data, sizeof(FLidarPointCloudPoint)); Dest += sizeof(FLidarPointCloudPoint);
+		FMemory::Memcpy(Dest, Data, sizeof(FLidarPointCloudPoint)); Dest += sizeof(FLidarPointCloudPoint);
+		FMemory::Memcpy(Dest, Data, sizeof(FLidarPointCloudPoint)); Dest += sizeof(FLidarPointCloudPoint);
+	}
+
+	RHIUnlockBuffer(VertexBufferRHI);
+}
+
 void FLidarPointCloudVertexFactory::InitRHI()
 {
 	VertexBuffer.InitResource();
 
 	FVertexDeclarationElementList Elements;
-	Elements.Add(AccessStreamComponent(FVertexStreamComponent(&VertexBuffer, 0, 0, VET_Float3), 0));
+	Elements.Add(AccessStreamComponent(FVertexStreamComponent(&VertexBuffer, 0, sizeof(FLidarPointCloudPoint), VET_Float3), 0));
+	Elements.Add(AccessStreamComponent(FVertexStreamComponent(&VertexBuffer, 12, sizeof(FLidarPointCloudPoint), VET_Color), 1));
+	Elements.Add(AccessStreamComponent(FVertexStreamComponent(&VertexBuffer, 16, sizeof(FLidarPointCloudPoint), VET_UInt), 2));
 	InitDeclaration(Elements);
 }
 
@@ -236,11 +270,37 @@ void FLidarPointCloudVertexFactory::ReleaseRHI()
 	VertexBuffer.ReleaseResource();
 }
 
+void FLidarPointCloudSharedVertexFactory::FPointCloudVertexBuffer::InitRHI()
+{
+	FRHIResourceCreateInfo CreateInfo(TEXT("FPointCloudVertexBuffer"));
+	VertexBufferRHI = RHICreateBuffer(sizeof(FVector), BUF_Static | BUF_VertexBuffer | BUF_ZeroStride, 0, ERHIAccess::VertexOrIndexBuffer, CreateInfo);
+	void* Buffer = RHILockBuffer(VertexBufferRHI, 0, sizeof(FVector), RLM_WriteOnly);
+	FMemory::Memzero(Buffer, sizeof(FVector));
+	RHIUnlockBuffer(VertexBufferRHI);
+}
+
+void FLidarPointCloudSharedVertexFactory::InitRHI()
+{
+	VertexBuffer.InitResource();
+
+	FVertexDeclarationElementList Elements;
+	Elements.Add(AccessStreamComponent(FVertexStreamComponent(&VertexBuffer, 0, 0, VET_Float3), 0));
+	Elements.Add(AccessStreamComponent(FVertexStreamComponent(&VertexBuffer, 0, 0, VET_Color), 1));
+	Elements.Add(AccessStreamComponent(FVertexStreamComponent(&VertexBuffer, 0, 0, VET_Color), 2));
+	InitDeclaration(Elements);
+}
+
+void FLidarPointCloudSharedVertexFactory::ReleaseRHI()
+{
+	FVertexFactory::ReleaseRHI();
+	VertexBuffer.ReleaseResource();
+}
+
 IMPLEMENT_TYPE_LAYOUT(FLidarPointCloudVertexFactoryShaderParameters);
 
-IMPLEMENT_VERTEX_FACTORY_PARAMETER_TYPE(FLidarPointCloudVertexFactory, SF_Vertex, FLidarPointCloudVertexFactoryShaderParameters);
+IMPLEMENT_VERTEX_FACTORY_PARAMETER_TYPE(FLidarPointCloudVertexFactoryBase, SF_Vertex, FLidarPointCloudVertexFactoryShaderParameters);
 
-IMPLEMENT_VERTEX_FACTORY_TYPE(FLidarPointCloudVertexFactory, "/Plugin/LidarPointCloud/Private/LidarPointCloudVertexFactory.ush",
+IMPLEMENT_VERTEX_FACTORY_TYPE(FLidarPointCloudVertexFactoryBase, "/Plugin/LidarPointCloud/Private/LidarPointCloudVertexFactory.ush",
 	  EVertexFactoryFlags::UsedWithMaterials
 	| EVertexFactoryFlags::SupportsDynamicLighting
 	| EVertexFactoryFlags::SupportsPositionOnly
diff --git a/Engine/Plugins/Enterprise/LidarPointCloud/Source/LidarPointCloudRuntime/Private/Rendering/LidarPointCloudRenderBuffers.h b/Engine/Plugins/Enterprise/LidarPointCloud/Source/LidarPointCloudRuntime/Private/Rendering/LidarPointCloudRenderBuffers.h
--- a/Engine/Plugins/Enterprise/LidarPointCloud/Source/LidarPointCloudRuntime/Private/Rendering/LidarPointCloudRenderBuffers.h
+++ b/Engine/Plugins/Enterprise/LidarPointCloud/Source/LidarPointCloudRuntime/Private/Rendering/LidarPointCloudRenderBuffers.h
@@ -62,19 +62,22 @@ public:
  */
 struct FLidarPointCloudBatchElementUserData
 {
+	FRHIShaderResourceView* TreeBuffer;
 	FRHIShaderResourceView* DataBuffer;
 	int32 bEditorView;
 	FVector SelectionColor;
-	int32 IndexDivisor;
 	FVector LocationOffset;
+	float RootCellSize;
+	FVector RootExtent;
+	int32 bUsePerPointScaling;
 	float VirtualDepth;
-	float SpriteSize;
-	int32 bUseLODColoration;
 	float SpriteSizeMultiplier;
+	float ReversedVirtualDepthMultiplier;
 	FVector ViewRightVector;
 	FVector ViewUpVector;
 	int32 bUseCameraFacing;
 	int32 bUseScreenSizeScaling;
+	int32 bUseStaticBuffers;
 	FVector BoundsSize;
 	FVector ElevationColorBottom;
 	FVector ElevationColorTop;
@@ -110,19 +113,22 @@ public:
 	void GetElementShaderBindings(const class FSceneInterface* Scene, const FSceneView* View, const FMeshMaterialShader* Shader, const EVertexInputStreamType InputStreamType, ERHIFeatureLevel::Type FeatureLevel,
 		const FVertexFactory* VertexFactory, const FMeshBatchElement& BatchElement, class FMeshDrawSingleShaderBindings& ShaderBindings, FVertexInputStreamArray& VertexStreams) const;
 
+	LAYOUT_FIELD(FShaderResourceParameter, TreeBuffer);
 	LAYOUT_FIELD(FShaderResourceParameter, DataBuffer);
 	LAYOUT_FIELD(FShaderParameter, bEditorView);
 	LAYOUT_FIELD(FShaderParameter, SelectionColor);
-	LAYOUT_FIELD(FShaderParameter, IndexDivisor);
 	LAYOUT_FIELD(FShaderParameter, LocationOffset);
+	LAYOUT_FIELD(FShaderParameter, RootCellSize);
+	LAYOUT_FIELD(FShaderParameter, RootExtent);
+	LAYOUT_FIELD(FShaderParameter, bUsePerPointScaling);
 	LAYOUT_FIELD(FShaderParameter, VirtualDepth);
-	LAYOUT_FIELD(FShaderParameter, SpriteSize);
-	LAYOUT_FIELD(FShaderParameter, bUseLODColoration);
 	LAYOUT_FIELD(FShaderParameter, SpriteSizeMultiplier);
+	LAYOUT_FIELD(FShaderParameter, ReversedVirtualDepthMultiplier);
 	LAYOUT_FIELD(FShaderParameter, ViewRightVector);
 	LAYOUT_FIELD(FShaderParameter, ViewUpVector);
 	LAYOUT_FIELD(FShaderParameter, bUseCameraFacing);
 	LAYOUT_FIELD(FShaderParameter, bUseScreenSizeScaling);
+	LAYOUT_FIELD(FShaderParameter, bUseStaticBuffers);
 	LAYOUT_FIELD(FShaderParameter, BoundsSize);
 	LAYOUT_FIELD(FShaderParameter, ElevationColorBottom);
 	LAYOUT_FIELD(FShaderParameter, ElevationColorTop);
@@ -142,33 +148,48 @@ public:
 	LAYOUT_FIELD(FShaderParameter, bStartClipped);
 };
 
-/**
- * Implementation of the custom Vertex Factory, containing only a ZeroStride position stream.
- */
-class FLidarPointCloudVertexFactory : public FVertexFactory
+class FLidarPointCloudVertexFactoryBase : public FVertexFactory
 {
-	DECLARE_VERTEX_FACTORY_TYPE(FLidarPointCloudVertexFactory);
+	DECLARE_VERTEX_FACTORY_TYPE(FLidarPointCloudVertexFactoryBase);
 
 public:
 	static bool ShouldCache(const FVertexFactoryShaderPermutationParameters& Parameters) { return IsFeatureLevelSupported(Parameters.Platform, ERHIFeatureLevel::SM5); }
-    static bool ShouldCompilePermutation(const FVertexFactoryShaderPermutationParameters& Parameters);
+	static bool ShouldCompilePermutation(const FVertexFactoryShaderPermutationParameters& Parameters);
 
-	FLidarPointCloudVertexFactory() : FVertexFactory(ERHIFeatureLevel::SM5) { }
+	FLidarPointCloudVertexFactoryBase() : FVertexFactory(ERHIFeatureLevel::SM5) { }
+};
+
+class FLidarPointCloudVertexFactory : public FLidarPointCloudVertexFactoryBase
+{
+public:
+	void Initialize(FLidarPointCloudPoint* Data, int32 NumPoints);
 
 private:
-	/** Very simple implementation of a ZeroStride Vertex Buffer */
 	class FPointCloudVertexBuffer : public FVertexBuffer
 	{
+		FLidarPointCloudPoint* Data;
+		int32 NumPoints;
+
 	public:
-		virtual void InitRHI() override
-		{
-			FRHIResourceCreateInfo CreateInfo(TEXT("FPointCloudVertexBuffer"));
-			VertexBufferRHI = RHICreateBuffer(sizeof(FVector), BUF_Static | BUF_VertexBuffer | BUF_ZeroStride, 0, ERHIAccess::VertexOrIndexBuffer, CreateInfo);
-			void* Buffer = RHILockBuffer(VertexBufferRHI, 0, sizeof(FVector), RLM_WriteOnly);
-			FMemory::Memzero(Buffer, sizeof(FVector));
-			RHIUnlockBuffer(VertexBufferRHI);
-		}
+		virtual void InitRHI() override;
+		virtual FString GetFriendlyName() const override { return TEXT("FPointCloudVertexBuffer"); }
+		friend FLidarPointCloudVertexFactory;
+	} VertexBuffer;
 
+	virtual void InitRHI() override;
+	virtual void ReleaseRHI() override;
+};
+
+/**
+ * Implementation of the custom Vertex Factory, containing only a ZeroStride position stream.
+ */
+class FLidarPointCloudSharedVertexFactory : public FLidarPointCloudVertexFactoryBase
+{
+	/** Very simple implementation of a ZeroStride Vertex Buffer */
+	class FPointCloudVertexBuffer : public FVertexBuffer
+	{
+	public:
+		virtual void InitRHI() override;
 		virtual FString GetFriendlyName() const override { return TEXT("FPointCloudVertexBuffer"); }
 	} VertexBuffer;
 
@@ -178,4 +199,4 @@ private:
 
 /** A set of global render resources shared between all Lidar Point Cloud proxies */
 extern TGlobalResource<FLidarPointCloudIndexBuffer> GLidarPointCloudIndexBuffer;
-extern TGlobalResource<FLidarPointCloudVertexFactory> GLidarPointCloudVertexFactory;
\ No newline at end of file
+extern TGlobalResource<FLidarPointCloudSharedVertexFactory> GLidarPointCloudSharedVertexFactory;
\ No newline at end of file
diff --git a/Engine/Plugins/Enterprise/LidarPointCloud/Source/LidarPointCloudRuntime/Private/Rendering/LidarPointCloudRendering.cpp b/Engine/Plugins/Enterprise/LidarPointCloud/Source/LidarPointCloudRuntime/Private/Rendering/LidarPointCloudRendering.cpp
--- a/Engine/Plugins/Enterprise/LidarPointCloud/Source/LidarPointCloudRuntime/Private/Rendering/LidarPointCloudRendering.cpp
+++ b/Engine/Plugins/Enterprise/LidarPointCloud/Source/LidarPointCloudRuntime/Private/Rendering/LidarPointCloudRendering.cpp
@@ -163,13 +163,15 @@ public:
 	FLidarPointCloudSceneProxy(ULidarPointCloudComponent* Component)
 		: FPrimitiveSceneProxy(Component)
 		, ProxyWrapper(MakeShared<FLidarPointCloudSceneProxyWrapper, ESPMode::ThreadSafe>(this))
-		, Component(Component)
+		, bCompatiblePlatform(GetScene().GetFeatureLevel() >= ERHIFeatureLevel::SM5)
 		, Owner(Component->GetOwner())
 		, CollisionRendering(Component->GetPointCloud()->CollisionRendering)
 	{
 		// Skip material verification - async update could occasionally cause it to crash
 		bVerifyUsedMaterials = false;
 
+		TreeBuffer = new FLidarPointCloudRenderBuffer();
+
 		MaterialRelevance = Component->GetMaterialRelevance(GetScene().GetFeatureLevel());
 	}
 	virtual ~FLidarPointCloudSceneProxy()
@@ -182,12 +184,12 @@ public:
 	{
 		QUICK_SCOPE_CYCLE_COUNTER(STAT_PointCloudSceneProxy_GetDynamicMeshElements);
 
-		if (!Component->GetMaterial(0))
+		if (!CanBeRendered() || !RenderData.RenderParams.Material)
 		{
 			return;
 		}
 
-		const bool bUsesSprites = Component->PointSize > 0;
+		const bool bUsesSprites = RenderData.RenderParams.PointSize > 0;
 
 		for (int32 ViewIndex = 0; ViewIndex < Views.Num(); ViewIndex++)
 		{
@@ -203,15 +205,15 @@ public:
 
 					for (const FLidarPointCloudProxyUpdateDataNode& Node : RenderData.SelectedNodes)
 					{
-						if (Node.DataNode && Node.DataNode->GetDataCache())
+						if (Node.DataNode && ((RenderData.bUseStaticBuffers && Node.DataNode->GetVertexFactory()) || (!RenderData.bUseStaticBuffers && Node.DataNode->GetDataCache())))
 						{
 							FMeshBatch& MeshBatch = Collector.AllocateMesh();
 
 							MeshBatch.Type = bUsesSprites ? PT_TriangleList : PT_PointList;
 							MeshBatch.LODIndex = 0;
-							MeshBatch.VertexFactory = &GLidarPointCloudVertexFactory;
+							MeshBatch.VertexFactory = RenderData.bUseStaticBuffers ? Node.DataNode->GetVertexFactory() : (FVertexFactory*)&GLidarPointCloudSharedVertexFactory;
 							MeshBatch.bWireframe = false;
-							MeshBatch.MaterialRenderProxy = Component->GetMaterial(0)->GetRenderProxy();
+							MeshBatch.MaterialRenderProxy = RenderData.RenderParams.Material->GetRenderProxy();
 							MeshBatch.ReverseCulling = IsLocalToWorldDeterminantNegative();
 							MeshBatch.DepthPriorityGroup = SDPG_World;
 
@@ -234,7 +236,7 @@ public:
 				FPrimitiveDrawInterface* PDI = Collector.GetPDI(ViewIndex);
 
 				// Draw selected nodes' bounds
-				if (Component->bDrawNodeBounds)
+				if (RenderData.RenderParams.bDrawNodeBounds)
 				{
 					for (const FBox& Node : RenderData.Bounds)
 					{
@@ -283,14 +285,17 @@ public:
 	{
 		FPrimitiveViewRelevance Result;
 
-		Result.bDrawRelevance = IsShown(View);
-		Result.bShadowRelevance = IsShadowCast(View);
-		Result.bDynamicRelevance = true;
-		Result.bStaticRelevance = false;
-		Result.bRenderInMainPass = ShouldRenderInMainPass();
-		Result.bUsesLightingChannels = GetLightingChannelMask() != GetDefaultLightingChannelMask();
-		Result.bRenderCustomDepth = ShouldRenderCustomDepth();
-		MaterialRelevance.SetPrimitiveViewRelevance(Result);
+		if (CanBeRendered())
+		{
+			Result.bDrawRelevance = IsShown(View);
+			Result.bShadowRelevance = IsShadowCast(View);
+			Result.bDynamicRelevance = true;
+			Result.bStaticRelevance = false;
+			Result.bRenderInMainPass = ShouldRenderInMainPass();
+			Result.bUsesLightingChannels = GetLightingChannelMask() != GetDefaultLightingChannelMask();
+			Result.bRenderCustomDepth = ShouldRenderCustomDepth();
+			MaterialRelevance.SetPrimitiveViewRelevance(Result);
+		}
 
 		return Result;
 	}
@@ -300,58 +305,51 @@ public:
 	{	
 		FLidarPointCloudBatchElementUserData UserDataElement;
 
-		const bool bUsesSprites = Component->PointSize > 0;
-
-		FVector BoundsSize = Component->GetPointCloud()->GetBounds().GetSize();
-
-		// Make sure to apply minimum bounds size
-		BoundsSize.X = FMath::Max<FVector::FReal>(BoundsSize.X, 0.001);
-		BoundsSize.Y = FMath::Max<FVector::FReal>(BoundsSize.Y, 0.001);
-		BoundsSize.Z = FMath::Max<FVector::FReal>(BoundsSize.Z, 0.001);
+		const bool bUsesSprites = RenderData.RenderParams.PointSize > 0;
 
 		// Update shader parameters
-		UserDataElement.bEditorView = Component->IsOwnedByEditor();
+		UserDataElement.bEditorView = RenderData.RenderParams.bOwnedByEditor;
+		UserDataElement.ReversedVirtualDepthMultiplier = RenderData.VDMultiplier;
 		UserDataElement.VirtualDepth = RenderData.VDMultiplier * Node.VirtualDepth;
-		UserDataElement.SpriteSize = RenderData.RootCellSize / FMath::Pow(2.0f, UserDataElement.VirtualDepth);
-		UserDataElement.SpriteSizeMultiplier = (Component->PointSize + Component->GapFillingStrength * 0.005f) * Component->GetComponentScale().GetAbsMax();
-		UserDataElement.bUseScreenSizeScaling = Component->bUseScreenSizeScaling;
-
-		UserDataElement.IndexDivisor = bUsesSprites ? 4 : 1;
-		UserDataElement.LocationOffset = Component->GetPointCloud()->GetLocationOffset().ToVector();		
+		UserDataElement.SpriteSizeMultiplier = bUsesSprites ? (RenderData.RenderParams.PointSize + RenderData.RenderParams.GapFillingStrength * 0.005f) * RenderData.RenderParams.ComponentScale : 0;
+		UserDataElement.bUseScreenSizeScaling = RenderData.RenderParams.ScalingMethod == ELidarPointCloudScalingMethod::FixedScreenSize;;
+		UserDataElement.bUsePerPointScaling = RenderData.RenderParams.ScalingMethod == ELidarPointCloudScalingMethod::PerPoint;
+		UserDataElement.bUseStaticBuffers = RenderData.bUseStaticBuffers;
+		UserDataElement.RootCellSize = RenderData.RootCellSize;
+		UserDataElement.RootExtent = FVector(RenderData.RenderParams.BoundsSize.GetAbsMax() * 0.5f);
+
+		UserDataElement.LocationOffset = RenderData.RenderParams.LocationOffset;
 		UserDataElement.ViewRightVector = InView->GetViewRight();
 		UserDataElement.ViewUpVector = InView->GetViewUp();
-		UserDataElement.bUseCameraFacing = !Component->ShouldRenderFacingNormals();
-		UserDataElement.BoundsSize = BoundsSize;
-		UserDataElement.ElevationColorBottom = FVector(Component->ColorSource == ELidarPointCloudColorationMode::None ? FColor::White : Component->ElevationColorBottom);
-		UserDataElement.ElevationColorTop = FVector(Component->ColorSource == ELidarPointCloudColorationMode::None ? FColor::White : Component->ElevationColorTop);
-		UserDataElement.bUseCircle = bUsesSprites && Component->GetPointShape() == ELidarPointCloudSpriteShape::Circle;
-		UserDataElement.bUseColorOverride = Component->ColorSource != ELidarPointCloudColorationMode::Data;
-		UserDataElement.bUseElevationColor = Component->ColorSource == ELidarPointCloudColorationMode::Elevation || Component->ColorSource == ELidarPointCloudColorationMode::None;
-		UserDataElement.Offset = Component->Offset;
-		UserDataElement.Contrast = Component->Contrast;
-		UserDataElement.Saturation = Component->Saturation;
-		UserDataElement.Gamma = Component->Gamma;
-		UserDataElement.Tint = FVector(Component->ColorTint);
-		UserDataElement.IntensityInfluence = Component->IntensityInfluence;
-
-		UserDataElement.bUseClassification = Component->ColorSource == ELidarPointCloudColorationMode::Classification;
-		UserDataElement.SetClassificationColors(Component->ClassificationColors);
+		UserDataElement.bUseCameraFacing = !RenderData.RenderParams.bShouldRenderFacingNormals;
+		UserDataElement.BoundsSize = RenderData.RenderParams.BoundsSize;
+		UserDataElement.ElevationColorBottom = FVector(RenderData.RenderParams.ColorSource == ELidarPointCloudColorationMode::None ? FColor::White : RenderData.RenderParams.ElevationColorBottom);
+		UserDataElement.ElevationColorTop = FVector(RenderData.RenderParams.ColorSource == ELidarPointCloudColorationMode::None ? FColor::White : RenderData.RenderParams.ElevationColorTop);
+		UserDataElement.bUseCircle = bUsesSprites && RenderData.RenderParams.PointShape == ELidarPointCloudSpriteShape::Circle;
+		UserDataElement.bUseColorOverride = RenderData.RenderParams.ColorSource != ELidarPointCloudColorationMode::Data;
+		UserDataElement.bUseElevationColor = RenderData.RenderParams.ColorSource == ELidarPointCloudColorationMode::Elevation || RenderData.RenderParams.ColorSource == ELidarPointCloudColorationMode::None;
+		UserDataElement.Offset = RenderData.RenderParams.Offset;
+		UserDataElement.Contrast = RenderData.RenderParams.Contrast;
+		UserDataElement.Saturation = RenderData.RenderParams.Saturation;
+		UserDataElement.Gamma = RenderData.RenderParams.Gamma;
+		UserDataElement.Tint = RenderData.RenderParams.ColorTint;
+		UserDataElement.IntensityInfluence = RenderData.RenderParams.IntensityInfluence;
+
+		UserDataElement.bUseClassification = RenderData.RenderParams.ColorSource == ELidarPointCloudColorationMode::Classification;
+		UserDataElement.SetClassificationColors(RenderData.RenderParams.ClassificationColors);
 				
 		UserDataElement.NumClippingVolumes = FMath::Min(RenderData.ClippingVolumes.Num(), 16);
 
 		for (uint32 i = 0; i < UserDataElement.NumClippingVolumes; ++i)
 		{
-			const ALidarClippingVolume* ClippingVolume = RenderData.ClippingVolumes[i];
-			const FVector Extent = ClippingVolume->GetActorScale3D() * 100;
-			UserDataElement.ClippingVolume[i] = FMatrix(FPlane(ClippingVolume->GetActorLocation(), ClippingVolume->Mode == ELidarClippingVolumeMode::ClipInside),
-													 FPlane(ClippingVolume->GetActorForwardVector(), Extent.X),
-												     FPlane(ClippingVolume->GetActorRightVector(), Extent.Y),
-													 FPlane(ClippingVolume->GetActorUpVector(), Extent.Z));
-
-			UserDataElement.bStartClipped = UserDataElement.bStartClipped || ClippingVolume->Mode == ELidarClippingVolumeMode::ClipOutside;
+			const FLidarPointCloudClippingVolumeParams& ClippingVolume = RenderData.ClippingVolumes[i];
+
+			UserDataElement.ClippingVolume[i] = ClippingVolume.PackedShaderData;
+			UserDataElement.bStartClipped |= ClippingVolume.Mode == ELidarClippingVolumeMode::ClipOutside;
 		}
 
-		UserDataElement.DataBuffer = Node.DataNode->GetDataCache()->SRV;
+		UserDataElement.DataBuffer = Node.DataNode->GetDataCache() ? Node.DataNode->GetDataCache()->SRV : nullptr;
+		UserDataElement.TreeBuffer = TreeBuffer->SRV;
 
 		return UserDataElement;
 	}
@@ -359,6 +357,7 @@ public:
 	virtual bool CanBeOccluded() const override { return !MaterialRelevance.bDisableDepthTest; }
 	virtual uint32 GetMemoryFootprint() const override { return(sizeof(*this) + GetAllocatedSize()); }
 	uint32 GetAllocatedSize() const { return(FPrimitiveSceneProxy::GetAllocatedSize()); }
+	bool CanBeRendered() const { return bCompatiblePlatform; }
 
 	virtual SIZE_T GetTypeHash() const override
 	{
@@ -366,7 +365,20 @@ public:
 		return reinterpret_cast<size_t>(&UniquePointer);
 	}
 
-	virtual void UpdateRenderData(FLidarPointCloudProxyUpdateData InRenderData) override { RenderData = InRenderData; }
+	virtual void UpdateRenderData(const FLidarPointCloudProxyUpdateData& InRenderData) override
+	{
+		RenderData = InRenderData;
+
+		const int32 NumTreeStructure = RenderData.TreeStructure.Num() > 0 ? RenderData.TreeStructure.Num() : 16;
+		TreeBuffer->Resize(NumTreeStructure);
+		uint8* DataPtr = (uint8*)RHILockBuffer(TreeBuffer->Buffer, 0, NumTreeStructure * sizeof(uint32), RLM_WriteOnly);
+		FMemory::Memzero(DataPtr, NumTreeStructure * sizeof(uint32));
+		if (RenderData.TreeStructure.Num() > 0)
+		{
+			FMemory::Memcpy(DataPtr, RenderData.TreeStructure.GetData(), NumTreeStructure * sizeof(uint32));
+		}
+		RHIUnlockBuffer(TreeBuffer->Buffer);
+	}
 
 public:
 	TSharedPtr<FLidarPointCloudSceneProxyWrapper, ESPMode::ThreadSafe> ProxyWrapper;
@@ -374,8 +386,11 @@ public:
 private:
 	FLidarPointCloudProxyUpdateData RenderData;
 
-	ULidarPointCloudComponent* Component;
+	FLidarPointCloudRenderBuffer* TreeBuffer;
 	FMaterialRelevance MaterialRelevance;
+	
+	bool bCompatiblePlatform;
+	
 	AActor* Owner;
 
 	FLidarPointCloudCollisionRendering* CollisionRendering;
@@ -387,7 +402,11 @@ FPrimitiveSceneProxy* ULidarPointCloudComponent::CreateSceneProxy()
 	if (PointCloud)
 	{
 		Proxy = new FLidarPointCloudSceneProxy(this);
-		FLidarPointCloudLODManager::RegisterProxy(this, Proxy->ProxyWrapper);
+		
+		if (Proxy->CanBeRendered())
+		{
+			FLidarPointCloudLODManager::RegisterProxy(this, Proxy->ProxyWrapper);
+		}
 	}
 	return Proxy;
 }
diff --git a/Engine/Plugins/Enterprise/LidarPointCloud/Source/LidarPointCloudRuntime/Public/LidarPointCloud.h b/Engine/Plugins/Enterprise/LidarPointCloud/Source/LidarPointCloudRuntime/Public/LidarPointCloud.h
--- a/Engine/Plugins/Enterprise/LidarPointCloud/Source/LidarPointCloudRuntime/Public/LidarPointCloud.h
+++ b/Engine/Plugins/Enterprise/LidarPointCloud/Source/LidarPointCloudRuntime/Public/LidarPointCloud.h
@@ -90,6 +90,18 @@ public:
 	UPROPERTY(EditAnywhere, BlueprintReadWrite, Category = "Normals", meta = (ClampMin = "0.0", DisplayName = "Noise Tolerance"))
 	float NormalsNoiseTolerance;
 
+private:
+	/**
+	 * Disables the LOD pipeline, allowing for much faster data operations (insert/remove/set)
+	 * at a potential expense of runtime performance. The whole asset will be treated as a single,
+	 * large asset with no granular density control, nor occlusion culling.
+	 * 
+	 * Recommended for assets, which have their data updated per-frame (such as live streaming).
+	 */
+	UPROPERTY(EditAnywhere, Category = "Performance", meta = (AllowPrivateAccess = "true"))
+	bool bOptimizedForDynamicData;
+
+public:
 	/** Holds pointer to the Import Settings used for the import. */
 	TSharedPtr<struct FLidarPointCloudImportSettings> ImportSettings;
 
@@ -183,6 +195,9 @@ public:
 	UFUNCTION(BlueprintPure, Category = "Lidar Point Cloud")
 	int32 GetNumNodes() const { return Octree.GetNumNodes(); }
 
+	UFUNCTION(BlueprintPure, Category = "Lidar Point Cloud")
+	FORCEINLINE float GetEstimatedPointSpacing() const { return Octree.GetEstimatedPointSpacing(); }
+
 	/** Returns the amount of memory in MB used to store the point cloud. */
 	UFUNCTION(BlueprintPure, Category = "Lidar Point Cloud")
 	int32 GetDataSize() const;
@@ -190,6 +205,9 @@ public:
 	UFUNCTION(BlueprintPure, Category = "Lidar Point Cloud")
 	FString GetSourcePath() const { return SourcePath.FilePath; }
 
+	UFUNCTION(BlueprintPure, Category = "Lidar Point Cloud")
+	bool IsOptimizedForDynamicData() const { return bOptimizedForDynamicData; }
+
 	UFUNCTION(BlueprintPure, Category = "Lidar Point Cloud")
 	FBox GetBounds() const { return Octree.GetBounds().ShiftBy(LocationOffset.ToVector()); }
 	FDoubleBox GetPreciseBounds(bool bUseOriginalCoordinates) const { return FDoubleBox(Octree.GetBounds()).ShiftBy(bUseOriginalCoordinates ? OriginalCoordinates : LocationOffset); }
@@ -256,11 +274,11 @@ public:
 	void GetPointsInBox(TArray64<FLidarPointCloudPoint*>& SelectedPoints, const FBox& Box, const bool& bVisibleOnly);
 
 	/**
-	 * Populates the array with the list of points within the given frustum.
-	 * Frustum is assumed to include the LocationOffset of the asset
+	 * Populates the array with the list of points within the given convex volume.
+	 * The volume is assumed to include the LocationOffset of the asset
 	 */
-	void GetPointsInFrustum(TArray<FLidarPointCloudPoint*>& SelectedPoints, const FConvexVolume& Frustum, const bool& bVisibleOnly);
-	void GetPointsInFrustum(TArray64<FLidarPointCloudPoint*>& SelectedPoints, const FConvexVolume& Frustum, const bool& bVisibleOnly);
+	void GetPointsInConvexVolume(TArray<FLidarPointCloudPoint*>& SelectedPoints, const FConvexVolume& ConvexVolume, const bool& bVisibleOnly);
+	void GetPointsInConvexVolume(TArray64<FLidarPointCloudPoint*>& SelectedPoints, const FConvexVolume& ConvexVolume, const bool& bVisibleOnly);
 
 	/**
 	 * Returns an array with copies of points from the tree
@@ -400,6 +418,9 @@ public:
 	UFUNCTION(BlueprintCallable, Category = "Lidar Point Cloud")
 	void SetSourcePath(const FString& NewSourcePath);
 
+	UFUNCTION(BlueprintCallable, Category = "Lidar Point Cloud")
+	void SetOptimizedForDynamicData(bool bNewOptimizedForDynamicData);
+
 	/**
 	 * Re-initializes the asset with new bounds.
 	 * Warning: Will erase all currently held data!
@@ -585,6 +606,13 @@ public:
 
 	UBodySetup* GetBodySetup();
 
+	//~ Begin Deprecated
+	UE_DEPRECATED(4.27, "Use GetPointsInConvexVolume instead.")
+	void GetPointsInFrustum(TArray<FLidarPointCloudPoint*>& SelectedPoints, const FConvexVolume& Frustum, const bool& bVisibleOnly);
+	UE_DEPRECATED(4.27, "Use GetPointsInConvexVolume instead.")
+	void GetPointsInFrustum(TArray64<FLidarPointCloudPoint*>& SelectedPoints, const FConvexVolume& Frustum, const bool& bVisibleOnly);
+	//~ End Deprecated
+
 public:
 	/** Aligns provided clouds based on the relative offset between their Original Coordinates. Retains overall centering of the group. */
 	static void AlignClouds(TArray<ULidarPointCloud*> PointCloudsToAlign);
@@ -633,7 +661,7 @@ private:
 	template <typename T>
 	void GetPointsInBox_Internal(TArray<FLidarPointCloudPoint*, T>& SelectedPoints, const FBox& Box, const bool& bVisibleOnly);
 	template <typename T>
-	void GetPointsInFrustum_Internal(TArray<FLidarPointCloudPoint*, T>& SelectedPoints, const FConvexVolume& Frustum, const bool& bVisibleOnly);
+	void GetPointsInConvexVolume_Internal(TArray<FLidarPointCloudPoint*, T>& SelectedPoints, const FConvexVolume& ConvexVolume, const bool& bVisibleOnly);
 	template <typename T>
 	void GetPointsAsCopies_Internal(TArray<FLidarPointCloudPoint, T>& Points, bool bReturnWorldSpace, int64 StartIndex = 0, int64 Count = -1) const;
 	template <typename T>
@@ -803,15 +831,6 @@ public:
 	static FLidarPointCloudNormal Conv_VectorToLidarPointCloudNormal(const FVector& Vector) { return FLidarPointCloudNormal(Vector); }
 };
 
-UENUM(BlueprintType)
-enum class ELidarClippingVolumeMode : uint8
-{
-	/** This will clip all points inside the volume */
-	ClipInside,
-	/** This will clip all points outside of the volume */
-	ClipOutside,
-};
-
 UCLASS(hidecategories = (Collision, Brush, Attachment, Physics, Volume, BrushBuilder), MinimalAPI)
 class ALidarClippingVolume : public AVolume
 {
@@ -833,4 +852,4 @@ public:
 	int32 Priority;
 
 	ALidarClippingVolume();
-};
\ No newline at end of file
+};
diff --git a/Engine/Plugins/Enterprise/LidarPointCloud/Source/LidarPointCloudRuntime/Public/LidarPointCloudComponent.h b/Engine/Plugins/Enterprise/LidarPointCloud/Source/LidarPointCloudRuntime/Public/LidarPointCloudComponent.h
--- a/Engine/Plugins/Enterprise/LidarPointCloud/Source/LidarPointCloudRuntime/Public/LidarPointCloudComponent.h
+++ b/Engine/Plugins/Enterprise/LidarPointCloud/Source/LidarPointCloudRuntime/Public/LidarPointCloudComponent.h
@@ -9,28 +9,6 @@
 
 class UBodySetup;
 
-UENUM(BlueprintType)
-enum class ELidarPointCloudColorationMode : uint8
-{
-	/** Uses color tint only */
-	None,
-	/** Uses imported RGB / Intensity data */
-	Data,
-	/** The cloud's color will be overridden with elevation-based color */
-	Elevation,
-	/** The cloud's color will be overridden with relative position-based color */
-	Position,
-	/** Uses Classification ID of the point along with the component's Classification Colors property to sample the color */
-	Classification
-};
-
-UENUM(BlueprintType)
-enum class ELidarPointCloudSpriteShape : uint8
-{
-	Square,
-	Circle,
-};
-
 UENUM(BlueprintType)
 enum class ELidarPointCloudSpriteOrientation : uint8
 {
@@ -58,21 +36,6 @@ private:
 	UMaterialInterface* CustomMaterial;
 
 public:
-	/**
-	 * Determines the minimum screen size for the node to be rendered.
-	 * Lower values will produce farther view distance, at the cost of higher CPU usage.
-	 */
-	UPROPERTY(EditAnywhere, BlueprintReadWrite, Category = "Performance", meta = (ClampMin = "0"))
-	float MinScreenSize;
-
-	/**
-	 * If enabled, points outside of the visible frustum will not be rendered.
-	 * While most project should leave this enabled, disabling it may help
-	 * with the data streaming lag when shooting cinematics.
-	 */
-	UPROPERTY(EditAnywhere, BlueprintReadWrite, AdvancedDisplay, Category = "Performance", meta = (ClampMin = "0"))
-	bool bUseFrustumCulling;
-
 	/**
 	 * Use to tweak the size of the points.
 	 * Set to 0 to switch to 1 pixel points.
@@ -80,12 +43,9 @@ public:
 	UPROPERTY(EditAnywhere, BlueprintReadWrite, Category = "Appearance", meta = (ClampMin = "0.0"))
 	float PointSize;
 
-	/**
-	 * If enabled, sprites will be rendered using screen-space scaling method.
-	 * In that mode, Point Size property will work as Screen Percentage.
-	 */
+	/** Determines how the points will be scaled  */
 	UPROPERTY(EditAnywhere, BlueprintReadWrite, Category = "Appearance")
-	bool bUseScreenSizeScaling;
+	ELidarPointCloudScalingMethod ScalingMethod;
 
 	/**
 	 * If set to > 0, it attempts to close gaps between points.
@@ -156,6 +116,14 @@ public:
 	UPROPERTY(EditAnywhere, BlueprintReadWrite, Category = "Color Adjustment", meta = (ClampMin = "0", ClampMax = "1"))
 	float IntensityInfluence;
 
+	/**
+	 * If enabled, points outside of the visible frustum will not be rendered.
+	 * While most project should leave this enabled, disabling it may help
+	 * with the data streaming lag when shooting cinematics.
+	 */
+	UPROPERTY(EditAnywhere, BlueprintReadWrite, Category = "Rendering")
+	bool bUseFrustumCulling;
+
 	/**
 	 * Minimum Depth from which the nodes should be rendered.
 	 * 0 to disable.
diff --git a/Engine/Plugins/Enterprise/LidarPointCloud/Source/LidarPointCloudRuntime/Public/LidarPointCloudLODManager.h b/Engine/Plugins/Enterprise/LidarPointCloud/Source/LidarPointCloudRuntime/Public/LidarPointCloudLODManager.h
--- a/Engine/Plugins/Enterprise/LidarPointCloud/Source/LidarPointCloudRuntime/Public/LidarPointCloudLODManager.h
+++ b/Engine/Plugins/Enterprise/LidarPointCloud/Source/LidarPointCloudRuntime/Public/LidarPointCloudLODManager.h
@@ -32,18 +32,6 @@ struct FLidarPointCloudViewData
 	bool ComputeFromEditorViewportClient(class FViewportClient* ViewportClient);
 };
 
-/** Convenience struct to group all selection params into one */
-struct FLidarPointCloudNodeSelectionParams
-{
-	float MinScreenSize;
-	float ScreenCenterImportance;
-	int32 MinDepth;
-	int32 MaxDepth;
-	float BoundsScale;
-	bool bUseFrustumCulling;
-	const TArray<const ALidarClippingVolume*>* ClippingVolumes;
-};
-
 /**  
  * This class is responsible for selecting nodes for rendering among all instances of all LidarPointCloud assets.
  * 
@@ -52,12 +40,14 @@ class FLidarPointCloudLODManager : public FTickableGameObject
 {
 	struct FRegisteredProxy
 	{
-		ULidarPointCloudComponent* Component;
-		ULidarPointCloud* PointCloud;
+		TWeakObjectPtr<ULidarPointCloudComponent> Component;
+		TWeakObjectPtr<ULidarPointCloud> PointCloud;
+		FLidarPointCloudOctree* Octree;
+		FLidarPointCloudComponentRenderParams ComponentRenderParams;
 		TWeakPtr<FLidarPointCloudSceneProxyWrapper, ESPMode::ThreadSafe> SceneProxyWrapper;
 		TSharedPtr<FLidarPointCloudTraversalOctree, ESPMode::ThreadSafe> TraversalOctree;
 
-		FRegisteredProxy(ULidarPointCloudComponent* Component, TWeakPtr<FLidarPointCloudSceneProxyWrapper, ESPMode::ThreadSafe> SceneProxyWrapper);
+		FRegisteredProxy(TWeakObjectPtr<ULidarPointCloudComponent> Component, TWeakPtr<FLidarPointCloudSceneProxyWrapper, ESPMode::ThreadSafe> SceneProxyWrapper);
 
 		/** Used to detect transform changes without the need of callbacks from the SceneProxy */
 		FTransform LastComponentTransform;
@@ -81,15 +71,6 @@ class FLidarPointCloudLODManager : public FTickableGameObject
 	FThreadSafeCounter64 NumPointsInFrustum;
 
 public:
-	/** Used for node size sorting and node selection. */
-	struct FNodeSizeData
-	{
-		FLidarPointCloudTraversalOctreeNode* Node;
-		float Size;
-		int32 ProxyIndex;
-		FNodeSizeData(FLidarPointCloudTraversalOctreeNode* Node, const float& Size, const int32& ProxyIndex) : Node(Node), Size(Size), ProxyIndex(ProxyIndex) {}
-	};
-
 	FLidarPointCloudLODManager();
 
 	virtual void Tick(float DeltaTime) override;
@@ -111,11 +92,11 @@ private:
 	 *
 	 * Returns the number of points in visible frustum
 	 */
-	int64 ProcessLOD(const TArray<FRegisteredProxy>& RegisteredProxies, const float CurrentTime, const uint32 PointBudget, const TArray<const ALidarClippingVolume*>& ClippingVolumes);
+	int64 ProcessLOD(const TArray<FRegisteredProxy>& RegisteredProxies, const float CurrentTime, const uint32 PointBudget, const TArray<FLidarPointCloudClippingVolumeParams>& ClippingVolumes);
 
 	/** Called to prepare the proxies for processing */
 	void PrepareProxies();
 
 	/** Compiles a list of all clipping volumes affecting any of the registered proxies */
-	TArray<const ALidarClippingVolume*> GetClippingVolumes() const;
+	TArray<FLidarPointCloudClippingVolumeParams> GetClippingVolumes() const;
 };
\ No newline at end of file
diff --git a/Engine/Plugins/Enterprise/LidarPointCloud/Source/LidarPointCloudRuntime/Public/LidarPointCloudOctree.h b/Engine/Plugins/Enterprise/LidarPointCloud/Source/LidarPointCloudRuntime/Public/LidarPointCloudOctree.h
--- a/Engine/Plugins/Enterprise/LidarPointCloud/Source/LidarPointCloudRuntime/Public/LidarPointCloudOctree.h
+++ b/Engine/Plugins/Enterprise/LidarPointCloud/Source/LidarPointCloudRuntime/Public/LidarPointCloudOctree.h
@@ -5,7 +5,6 @@
 #include "CoreMinimal.h"
 #include "LidarPointCloudShared.h"
 #include "LidarPointCloudSettings.h"
-#include "LidarPointCloudLODManager.h"
 #include "Meshing/LidarPointCloudMeshing.h"
 #include "HAL/ThreadSafeCounter64.h"
 #include "Misc/ScopeLock.h"
@@ -94,6 +93,7 @@ private:
 
 	/** Holds render data for this node */
 	class FLidarPointCloudRenderBuffer* DataCache;
+	class FLidarPointCloudVertexFactory* VertexFactory;
 
 	bool bRenderDataDirty;
 
@@ -121,7 +121,14 @@ public:
 	/** Returns a pointer to the point data */
 	FORCEINLINE FLidarPointCloudRenderBuffer* GetDataCache() { return DataCache; }
 
-	bool BuildDataCache();
+	/** Return a pointer to the vertex factory containing pre-cached geometry */
+	FORCEINLINE FLidarPointCloudVertexFactory* GetVertexFactory() { return VertexFactory; }
+
+	/**
+	 * Builds and updates the necessary render data buffers
+	 * Returns true if successful
+	 */
+	bool BuildDataCache(bool bUseStaticBuffers);
 
 	/** Returns the sum of grid and padding points allocated to this node. */
 	FORCEINLINE uint32 GetNumPoints() const { return NumPoints; }
@@ -138,6 +145,8 @@ public:
 	/** Returns a pointer to the node at the given location, or null if one doesn't exist yet. */
 	FLidarPointCloudOctreeNode* GetChildNodeAtLocation(const uint8& Location) const;
 
+	uint8 GetChildrenBitmask() const;
+
 	void UpdateNumVisiblePoints();
 
 	/** Attempts to insert given points to this node or passes it to the children, otherwise. */
@@ -156,11 +165,17 @@ public:
 	/** Returns true, if the node has its data loaded */
 	bool HasData() const { return bHasData; }
 
+	/** Returns true, if the node has its data loaded */
+	bool HasRenderData() const { return DataCache || VertexFactory; }
+
 	/**
 	 * Releases the BulkData
 	 * If forced, the node will be released even if persistent
 	 */
 	void ReleaseData(bool bForce = false);
+	
+	/** Releases and removes the render data cache */
+	void ReleaseDataCache();
 
 	/** Convenience function, to add point statistics to the Tree table. */
 	void AddPointCount(int32 PointCount);
@@ -169,6 +184,13 @@ public:
 	void SortVisiblePoints();
 
 private:
+	template <typename T>
+	void InsertPoints_Internal(T Points, const int64& Count, ELidarPointCloudDuplicateHandling DuplicateHandling, const FVector& Translation);
+	void InsertPoints_Dynamic(const FLidarPointCloudPoint* Points, const int64& Count, const FVector& Translation);
+	void InsertPoints_Static(const FLidarPointCloudPoint* Points, const int64& Count, ELidarPointCloudDuplicateHandling DuplicateHandling, const FVector& Translation);
+	void InsertPoints_Dynamic(FLidarPointCloudPoint** Points, const int64& Count, const FVector& Translation);
+	void InsertPoints_Static(FLidarPointCloudPoint** Points, const int64& Count, ELidarPointCloudDuplicateHandling DuplicateHandling, const FVector& Translation);
+
 	friend FLidarPointCloudOctree;
 	friend FLidarPointCloudTraversalOctree;
 	friend FLidarPointCloudTraversalOctreeNode;
@@ -346,9 +368,9 @@ public:
 	void GetPointsInBox(TArray<FLidarPointCloudPoint*>& SelectedPoints, const FBox& Box, const bool& bVisibleOnly);
 	void GetPointsInBox(TArray64<FLidarPointCloudPoint*>& SelectedPoints, const FBox& Box, const bool& bVisibleOnly);
 
-	/** Populates the array with the list of points within the given frustum. */
-	void GetPointsInFrustum(TArray<FLidarPointCloudPoint*>& SelectedPoints, const FConvexVolume& Frustum, const bool& bVisibleOnly);
-	void GetPointsInFrustum(TArray64<FLidarPointCloudPoint*>& SelectedPoints, const FConvexVolume& Frustum, const bool& bVisibleOnly);
+	/** Populates the array with the list of points within the given convex volume. */
+	void GetPointsInConvexVolume(TArray<FLidarPointCloudPoint*>& SelectedPoints, const FConvexVolume& ConvexVolume, const bool& bVisibleOnly);
+	void GetPointsInConvexVolume(TArray64<FLidarPointCloudPoint*>& SelectedPoints, const FConvexVolume& ConvexVolume, const bool& bVisibleOnly);
 
 	/** Populates the given array with copies of points from the tree */
 	void GetPointsAsCopies(TArray<FLidarPointCloudPoint>& SelectedPoints, const FTransform* LocalToWorld, int64 StartIndex = 0, int64 Count = -1) const;
@@ -442,8 +464,11 @@ public:
 	/** Marks render data of all nodes as dirty. */
 	void MarkRenderDataDirty();
 
-	/** Marks render data of all nodes within the given frustum as dirty. */
-	void MarkRenderDataInFrustumDirty(const FConvexVolume& Frustum);
+	/** Marks render data of all nodes within the given sphere as dirty. */
+	void MarkRenderDataInSphereDirty(const FSphere& Sphere);
+	
+	/** Marks render data of all nodes within the given convex volume as dirty. */
+	void MarkRenderDataInConvexVolumeDirty(const FConvexVolume& ConvexVolume);
 
 	/** Initializes the Octree properties. */
 	void Initialize(const FVector& InExtent);
@@ -503,13 +528,11 @@ public:
 	/** Removes the given traversal octree from the list */
 	void UnregisterTraversalOctree(FLidarPointCloudTraversalOctree* TraversalOctree);
 
-	/** If bImmediate is true, the node will be loaded immediately, otherwise, it will be queued for async loading. */
-	void QueueNode(FLidarPointCloudOctreeNode* Node, float Lifetime);
-
-	/** Streams all requested nodes */
-	void StreamQueuedNodes();
-
-	void UnloadOldNodes(const float& CurrentTime);
+	/**
+	 * Streams requested nodes or extends their lifetime, if already loaded
+	 * Unloads all unused nodes with expired lifetime
+	 */
+	void StreamNodes(TArray<FLidarPointCloudOctreeNode*>& Nodes, const float& CurrentTime);
 
 	/** Returns true, if the cloud is fully and persistently loaded. */
 	bool IsFullyLoaded() const { return bIsFullyLoaded; }
@@ -523,6 +546,21 @@ public:
 	 */
 	void ReleaseAllNodes(bool bIncludePersistent);
 
+	bool IsOptimizedForDynamicData() const;
+
+	void OptimizeForDynamicData();
+
+	void OptimizeForStaticData();
+
+	//~ Begin Deprecated
+	UE_DEPRECATED(4.27, "Use GetPointsInConvexVolume instead.")
+	void GetPointsInFrustum(TArray<FLidarPointCloudPoint*>& SelectedPoints, const FConvexVolume& Frustum, const bool& bVisibleOnly);
+	UE_DEPRECATED(4.27, "Use GetPointsInConvexVolume instead.")
+	void GetPointsInFrustum(TArray64<FLidarPointCloudPoint*>& SelectedPoints, const FConvexVolume& Frustum, const bool& bVisibleOnly);
+	UE_DEPRECATED(4.27, "Use MarkRenderDataInConvexVolumeDirty instead.")
+	void MarkRenderDataInFrustumDirty(const FConvexVolume& Frustum);
+	//~ End Deprecated
+
 private:
 	void RefreshAllocatedSize();
 
@@ -538,7 +576,7 @@ private:
 	template <typename T>
 	void GetPointsInBox_Internal(TArray<const FLidarPointCloudPoint*, T>& SelectedPoints, const FBox& Box, const bool& bVisibleOnly) const;
 	template <typename T>
-	void GetPointsInFrustum_Internal(TArray<FLidarPointCloudPoint*, T>& SelectedPoints, const FConvexVolume& Frustum, const bool& bVisibleOnly);
+	void GetPointsInConvexVolume_Internal(TArray<FLidarPointCloudPoint*, T>& SelectedPoints, const FConvexVolume& ConvexVolume, const bool& bVisibleOnly);
 	template <typename T>
 	void GetPointsAsCopies_Internal(TArray<FLidarPointCloudPoint, T>& Points, const FTransform* LocalToWorld, int64 StartIndex = 0, int64 Count = -1) const;
 	template <typename T>
@@ -584,8 +622,13 @@ struct FLidarPointCloudTraversalOctreeNode
 	/** Depth of this node */
 	uint8 Depth;
 
+	/** Calculated for use with adaptive sprite scaling */
+	uint8 VirtualDepth;
+
 	FLidarPointCloudTraversalOctreeNode* Parent;
 
+	FLidarPointCloudTraversalOctree* Octree;
+
 	/** Stores the children array */
 	TArray<FLidarPointCloudTraversalOctreeNode> Children;
 
@@ -597,12 +640,34 @@ struct FLidarPointCloudTraversalOctreeNode
 	FLidarPointCloudTraversalOctreeNode();
 
 	/** Builds the traversal version of the given node. */
-	void Build(FLidarPointCloudOctreeNode* Node, const FTransform& LocalToWorld, const FVector& LocationOffset);
+	void Build(FLidarPointCloudTraversalOctree* TraversalOctree, FLidarPointCloudOctreeNode* Node, const FTransform& LocalToWorld, const FVector& LocationOffset);
 
 	/** Calculates virtual depth of this node, to be used to estimate the best sprite size */
-	uint8 CalculateVirtualDepth(const TArray<float>& LevelWeights, const float& VDMultiplier, const float& PointSizeBias) const;
+	void CalculateVirtualDepth(const TArray<float>& LevelWeights, const float& PointSizeBias);
+
+	FORCEINLINE bool IsAvailable() const { return bSelected && DataNode->HasRenderData(); }
+};
 
-	FORCEINLINE bool IsAvailable() const { return bSelected && DataNode->HasData(); }
+/** Used for node size sorting and node selection. */
+struct FLidarPointCloudTraversalOctreeNodeSizeData
+{
+	FLidarPointCloudTraversalOctreeNode* Node;
+	float Size;
+	int32 ProxyIndex;
+
+	FLidarPointCloudTraversalOctreeNodeSizeData(FLidarPointCloudTraversalOctreeNode* Node, const float& Size, const int32& ProxyIndex);
+};
+
+/** Convenience struct to group all selection params into one */
+struct FLidarPointCloudNodeSelectionParams
+{
+	float MinScreenSize;
+	float ScreenCenterImportance;
+	int32 MinDepth;
+	int32 MaxDepth;
+	float BoundsScale;
+	bool bUseFrustumCulling;
+	const TArray<struct FLidarPointCloudClippingVolumeParams>* ClippingVolumes;
 };
 
 /**
@@ -639,8 +704,15 @@ struct FLidarPointCloudTraversalOctree
 	FLidarPointCloudTraversalOctree& operator=(const FLidarPointCloudTraversalOctree&) = delete;
 	FLidarPointCloudTraversalOctree& operator=(FLidarPointCloudTraversalOctree&&) = delete;
 
-	/** Selects and appends the subset of visible nodes for rendering. */
-	void GetVisibleNodes(TArray<FLidarPointCloudLODManager::FNodeSizeData>& NodeSizeData, const FLidarPointCloudViewData& ViewData, const int32& ProxyIndex, const FLidarPointCloudNodeSelectionParams& SelectionParams, const float& CurrentTime);
+	/**
+	 * Selects and appends the subset of visible nodes for rendering.
+	 * Returns number of selected nodes
+	 */
+	int32 GetVisibleNodes(TArray<FLidarPointCloudTraversalOctreeNodeSizeData>& NodeSizeData, const struct FLidarPointCloudViewData* ViewData, const int32& ProxyIndex, const FLidarPointCloudNodeSelectionParams& SelectionParams);
+
+	void CalculateVisibilityStructure(TArray<uint32>& OutData);
+
+	void CalculateLevelWeightsForSelectedNodes(TArray<float>& OutLevelWeights);
 
 	FVector GetCenter() const { return Root.Center; }
 	FVector GetExtent() const { return Extents[0]; }
diff --git a/Engine/Plugins/Enterprise/LidarPointCloud/Source/LidarPointCloudRuntime/Public/LidarPointCloudSettings.h b/Engine/Plugins/Enterprise/LidarPointCloud/Source/LidarPointCloudRuntime/Public/LidarPointCloudSettings.h
--- a/Engine/Plugins/Enterprise/LidarPointCloud/Source/LidarPointCloudRuntime/Public/LidarPointCloudSettings.h
+++ b/Engine/Plugins/Enterprise/LidarPointCloud/Source/LidarPointCloudRuntime/Public/LidarPointCloudSettings.h
@@ -90,6 +90,10 @@ public:
 	UPROPERTY(config, EditAnywhere, Category=Performance, meta = (EditCondition = "bUseAsyncImport"))
 	float RenderDataSmoothingMaxFrametime;
 
+	/** Enabling this will greatly improve runtime performance at a cost of quadrupling VRAM use */
+	UPROPERTY(config, EditAnywhere, Category=Performance)
+	bool bUseFastRendering;
+
 	/** Affects the size of per-thread data for the meshing algorithm. */
 	UPROPERTY(config, EditAnywhere, Category=Collision)
 	int32 MeshingBatchSize;
diff --git a/Engine/Plugins/Enterprise/LidarPointCloud/Source/LidarPointCloudRuntime/Public/LidarPointCloudShared.h b/Engine/Plugins/Enterprise/LidarPointCloud/Source/LidarPointCloudRuntime/Public/LidarPointCloudShared.h
--- a/Engine/Plugins/Enterprise/LidarPointCloud/Source/LidarPointCloudRuntime/Public/LidarPointCloudShared.h
+++ b/Engine/Plugins/Enterprise/LidarPointCloud/Source/LidarPointCloudRuntime/Public/LidarPointCloudShared.h
@@ -472,6 +472,36 @@ enum class ELidarPointCloudAsyncMode : uint8
 	Progress
 };
 
+UENUM(BlueprintType)
+enum class ELidarPointCloudScalingMethod : uint8
+{
+	/**
+	 * Points are scaled based on the estimated density of their containing node.
+	 * Recommended for assets with high variance of point densities, but may produce less fine detail overall.
+	 * Default method in 4.25 and 4.26
+	 */
+	PerNode,
+
+	/**
+	 * Similar to PerNode, but the density is calculated adaptively based on the current view.
+	 * Produces good amount of fine detail while being generally resistant to density variance.
+	 */
+	PerNodeAdaptive,
+
+	/**
+	 * Points are scaled based on their individual calculated depth.
+	 * Capable of resolving the highest amount of fine detail, but is the most susceptible to 
+	 * density changes across the dataset, and may result in patches of varying point sizes.
+	 */
+	PerPoint,
+
+	/**
+	 * Sprites will be rendered using screen-space scaling method.
+	 * In that mode, Point Size property will work as Screen Percentage.
+	 */
+	FixedScreenSize
+};
+
 /** Used to help track multiple buffer allocations */
 class LIDARPOINTCLOUDRUNTIME_API FLidarPointCloudDataBuffer
 {
@@ -533,7 +563,7 @@ private:
 };
 
 /** Used for Raycasting */
-struct FLidarPointCloudRay
+struct LIDARPOINTCLOUDRUNTIME_API FLidarPointCloudRay
 {
 public:
 	FVector Origin;
@@ -549,6 +579,11 @@ public:
 		SetDirection(Direction);
 	}
 
+	static FORCEINLINE FLidarPointCloudRay FromLocations(const FVector& Origin, const FVector& Destination)
+	{
+		return FLidarPointCloudRay(Origin, (Destination - Origin).GetSafeNormal());
+	}
+
 	FLidarPointCloudRay& TransformBy(const FTransform& Transform)
 	{
 		Origin = Transform.TransformPosition(Origin);
@@ -624,6 +659,94 @@ public:
 	}
 };
 
+UENUM(BlueprintType)
+enum class ELidarClippingVolumeMode : uint8
+{
+	/** This will clip all points inside the volume */
+	ClipInside,
+	/** This will clip all points outside of the volume */
+	ClipOutside,
+};
+
+/** Used to pass clipping information for async processing, to avoid accessing UObjects in non-GT */
+struct FLidarPointCloudClippingVolumeParams
+{
+	ELidarClippingVolumeMode Mode;
+	int32 Priority;
+	FBox Bounds;
+	FMatrix PackedShaderData;
+
+	FORCEINLINE bool operator<(const FLidarPointCloudClippingVolumeParams& O) const
+	{
+		return (Priority < O.Priority) || (Priority == O.Priority && Mode > O.Mode);
+	}
+
+	FLidarPointCloudClippingVolumeParams(const class ALidarClippingVolume* ClippingVolume);
+};
+
+UENUM(BlueprintType)
+enum class ELidarPointCloudColorationMode : uint8
+{
+	/** Uses color tint only */
+	None,
+	/** Uses imported RGB / Intensity data */
+	Data,
+	/** The cloud's color will be overridden with elevation-based color */
+	Elevation,
+	/** The cloud's color will be overridden with relative position-based color */
+	Position,
+	/** Uses Classification ID of the point along with the component's Classification Colors property to sample the color */
+	Classification
+};
+
+UENUM(BlueprintType)
+enum class ELidarPointCloudSpriteShape : uint8
+{
+	Square,
+	Circle,
+};
+
+/** Convenience struct to group all component's rendering params into one */
+struct FLidarPointCloudComponentRenderParams
+{
+	int32 MinDepth;
+	int32 MaxDepth;
+
+	float BoundsScale;
+	FVector BoundsSize;
+	FVector LocationOffset;
+	float ComponentScale;
+
+	float PointSize;
+	float PointSizeBias;
+	float GapFillingStrength;
+	
+	bool bOwnedByEditor;
+	bool bDrawNodeBounds;
+	bool bUseScreenSizeScaling;
+	bool bShouldRenderFacingNormals;
+	bool bUseFrustumCulling;
+
+	ELidarPointCloudColorationMode ColorSource;
+	ELidarPointCloudSpriteShape PointShape;
+	ELidarPointCloudScalingMethod ScalingMethod;
+
+	FVector4 Saturation;
+	FVector4 Contrast;
+	FVector4 Gamma;
+	FVector4 Offset;
+	FVector ColorTint;
+	float IntensityInfluence;
+
+	TMap<int32, FLinearColor> ClassificationColors;
+	FLinearColor ElevationColorBottom;
+	FLinearColor ElevationColorTop;
+
+	class UMaterialInterface* Material = nullptr;
+
+	void UpdateFromComponent(class ULidarPointCloudComponent* Component);
+};
+
 struct FBenchmarkTimer
 {
 	static void Reset()
diff --git a/Engine/Plugins/Enterprise/LidarPointCloud/Source/LidarPointCloudRuntime/Public/Rendering/LidarPointCloudRendering.h b/Engine/Plugins/Enterprise/LidarPointCloud/Source/LidarPointCloudRuntime/Public/Rendering/LidarPointCloudRendering.h
--- a/Engine/Plugins/Enterprise/LidarPointCloud/Source/LidarPointCloudRuntime/Public/Rendering/LidarPointCloudRendering.h
+++ b/Engine/Plugins/Enterprise/LidarPointCloud/Source/LidarPointCloudRuntime/Public/Rendering/LidarPointCloudRendering.h
@@ -3,6 +3,7 @@
 #pragma once
 
 #include "CoreMinimal.h"
+#include "LidarPointCloudShared.h"
 
 class ILidarPointCloudSceneProxy;
 
@@ -40,13 +41,19 @@ struct FLidarPointCloudProxyUpdateData
 
 	float VDMultiplier;
 	float RootCellSize;
+	
+	bool bUseStaticBuffers;
+
+	TArray<uint32> TreeStructure;
 
 #if !(UE_BUILD_SHIPPING)
 	/** Stores bounds of selected nodes, used for debugging */
 	TArray<FBox> Bounds;
 #endif
 
-	TArray<const class ALidarClippingVolume*> ClippingVolumes;
+	TArray<FLidarPointCloudClippingVolumeParams> ClippingVolumes;
+
+	FLidarPointCloudComponentRenderParams RenderParams;
 
 	FLidarPointCloudProxyUpdateData();
 };
@@ -56,5 +63,5 @@ class ILidarPointCloudSceneProxy
 {
 public:
 	/** Updates necessary render data for the proxy. Initiated via LOD Manager's Tick */
-	virtual void UpdateRenderData(FLidarPointCloudProxyUpdateData InRenderData) = 0;
+	virtual void UpdateRenderData(const FLidarPointCloudProxyUpdateData& InRenderData) = 0;
 };
\ No newline at end of file
