Crash occurs while Building Collision

diff --git a/Engine/Plugins/Enterprise/LidarPointCloud/Source/LidarPointCloudRuntime/Private/LidarPointCloud.cpp b/Engine/Plugins/Enterprise/LidarPointCloud/Source/LidarPointCloudRuntime/Private/LidarPointCloud.cpp
--- a/Engine/Plugins/Enterprise/LidarPointCloud/Source/LidarPointCloudRuntime/Private/LidarPointCloud.cpp
+++ b/Engine/Plugins/Enterprise/LidarPointCloud/Source/LidarPointCloudRuntime/Private/LidarPointCloud.cpp
@@ -1,7 +1,6 @@
 // Copyright Epic Games, Inc. All Rights Reserved.
 
 #include "LidarPointCloud.h"
-
 #include "LidarPointCloudShared.h"
 #include "LidarPointCloudActor.h"
 #include "LidarPointCloudComponent.h"
@@ -280,6 +279,7 @@ ULidarPointCloud::ULidarPointCloud()
 	, LocationOffset(FDoubleVector::ZeroVector)
 	, Notifications(this)
 	, BodySetup(nullptr)
+	, NewBodySetup(nullptr)
 	, bCollisionBuildInProgress(false)
 {
 	// Make sure we are transactional to allow undo redo
@@ -600,22 +600,22 @@ void ULidarPointCloud::BuildCollision()
 	bCollisionBuildInProgress = true;
 	MarkPackageDirty();
 
-	UBodySetup* NewBodySetup = NewObject<UBodySetup>(this);
+	NewBodySetup = NewObject<UBodySetup>(this);
 	NewBodySetup->BodySetupGuid = FGuid::NewGuid();
 	NewBodySetup->CollisionTraceFlag = CTF_UseComplexAsSimple;
 	NewBodySetup->bHasCookedCollisionData = true;
 
-	Async(EAsyncExecution::Thread, [this, Notification, NewBodySetup]{
+	Async(EAsyncExecution::Thread, [this, Notification]{
 		Octree.BuildCollision(MaxCollisionError, true);
 
 		FBenchmarkTimer::Reset();
 #if WITH_PHYSX  && PHYSICS_INTERFACE_PHYSX
-		AsyncTask(ENamedThreads::GameThread, [this, NewBodySetup, Notification] {
-			NewBodySetup->CreatePhysicsMeshesAsync(FOnAsyncPhysicsCookFinished::CreateUObject(this, &ULidarPointCloud::FinishPhysicsAsyncCook, NewBodySetup, Notification));
+		AsyncTask(ENamedThreads::GameThread, [this, Notification] {
+			NewBodySetup->CreatePhysicsMeshesAsync(FOnAsyncPhysicsCookFinished::CreateUObject(this, &ULidarPointCloud::FinishPhysicsAsyncCook, Notification));
 		});
 #elif WITH_CHAOS
 		NewBodySetup->CreatePhysicsMeshes();
-		AsyncTask(ENamedThreads::GameThread, [this, NewBodySetup, Notification] { FinishPhysicsAsyncCook(true, NewBodySetup, Notification); });
+		AsyncTask(ENamedThreads::GameThread, [this, Notification] { FinishPhysicsAsyncCook(true, Notification); });
 #endif		
 	});
 }
@@ -1377,7 +1377,7 @@ FBox ULidarPointCloud::CalculateBoundsFromPoints(FLidarPointCloudPoint** Points,
 	return Bounds;
 }
 
-void ULidarPointCloud::FinishPhysicsAsyncCook(bool bSuccess, UBodySetup* NewBodySetup, TSharedRef<FLidarPointCloudNotification, ESPMode::ThreadSafe> Notification)
+void ULidarPointCloud::FinishPhysicsAsyncCook(bool bSuccess, TSharedRef<FLidarPointCloudNotification, ESPMode::ThreadSafe> Notification)
 {
 	FBenchmarkTimer::Log("CookingCollision");
 	Notification->Close(bSuccess);
diff --git a/Engine/Plugins/Enterprise/LidarPointCloud/Source/LidarPointCloudRuntime/Private/LidarPointCloudLODManager.cpp b/Engine/Plugins/Enterprise/LidarPointCloud/Source/LidarPointCloudRuntime/Private/LidarPointCloudLODManager.cpp
--- a/Engine/Plugins/Enterprise/LidarPointCloud/Source/LidarPointCloudRuntime/Private/LidarPointCloudLODManager.cpp
+++ b/Engine/Plugins/Enterprise/LidarPointCloud/Source/LidarPointCloudRuntime/Private/LidarPointCloudLODManager.cpp
@@ -424,14 +424,14 @@ void FLidarPointCloudLODManager::Tick(float DeltaTime)
 	bProcessing = true;
 
 	Time += DeltaTime;
+	
+	LastPointBudget = GetPointBudget(NumPointsInFrustum.GetValue());
 
-	const uint32 PointBudget = GetPointBudget(NumPointsInFrustum.GetValue());
-
-	SET_DWORD_STAT(STAT_PointBudget, PointBudget);
+	SET_DWORD_STAT(STAT_PointBudget, LastPointBudget);
 
 	PrepareProxies();
 
-	Async(EAsyncExecution::ThreadPool, [this, CurrentRegisteredProxies = RegisteredProxies, PointBudget, ClippingVolumes = GetClippingVolumes()]
+	Async(EAsyncExecution::ThreadPool, [this, CurrentRegisteredProxies = RegisteredProxies, PointBudget = LastPointBudget, ClippingVolumes = GetClippingVolumes()]
 	{
 		NumPointsInFrustum.Set(ProcessLOD(CurrentRegisteredProxies, Time, PointBudget, ClippingVolumes));
 		bProcessing = false;
@@ -447,17 +447,22 @@ void FLidarPointCloudLODManager::RegisterProxy(ULidarPointCloudComponent* Compon
 {
 	if (IsValid(Component))
 	{
-		static FLidarPointCloudLODManager Instance;
 		static FCriticalSection ProxyLock;
 
 		if (TSharedPtr<FLidarPointCloudSceneProxyWrapper, ESPMode::ThreadSafe> SceneProxyWrapperShared = SceneProxyWrapper.Pin())
 		{
 			FScopeLock Lock(&ProxyLock);
-			Instance.RegisteredProxies.Emplace(Component, SceneProxyWrapper);
+			Get().RegisteredProxies.Emplace(Component, SceneProxyWrapper);
+			Get().ForceProcessLOD();
 		}
 	}
 }
 
+void FLidarPointCloudLODManager::RefreshLOD()
+{
+	Get().ForceProcessLOD();
+}
+
 int64 FLidarPointCloudLODManager::ProcessLOD(const TArray<FLidarPointCloudLODManager::FRegisteredProxy>& InRegisteredProxies, const float CurrentTime, const uint32 PointBudget, const TArray<FLidarPointCloudClippingVolumeParams>& ClippingVolumes)
 {
 	struct FSelectionFilterParams
@@ -647,11 +652,8 @@ int64 FLidarPointCloudLODManager::ProcessLOD(const TArray<FLidarPointCloudLODMan
 		}
 	}
 
-	// Set when to release the BulkData, if no longer visible
-	const float BulkDataLifetime = CurrentTime + GetDefault<ULidarPointCloudSettings>()->CachedNodeLifetime;
-
 	// Perform data streaming in a separate thread
-	Async(EAsyncExecution::ThreadPool, [OctreeStreamingMap = MoveTemp(OctreeStreamingMap), CurrentTime, BulkDataLifetime]() mutable
+	Async(EAsyncExecution::ThreadPool, [OctreeStreamingMap = MoveTemp(OctreeStreamingMap), CurrentTime]() mutable
 	{
 		SCOPE_CYCLE_COUNTER(STAT_NodeStreaming);
 
@@ -668,7 +670,7 @@ int64 FLidarPointCloudLODManager::ProcessLOD(const TArray<FLidarPointCloudLODMan
 	// Update Render Data
 	if (TotalPointsSelected > 0)
 	{
-		ENQUEUE_RENDER_COMMAND(ProcessLidarPointCloudLOD)([PointBudget, TotalPointsSelected, ProxyUpdateData](FRHICommandListImmediate& RHICmdList)
+		ENQUEUE_RENDER_COMMAND(ProcessLidarPointCloudLOD)([ProxyUpdateData](FRHICommandListImmediate& RHICmdList) mutable
 		{
 			SCOPE_CYCLE_COUNTER(STAT_UpdateRenderData);
 
@@ -679,18 +681,18 @@ int64 FLidarPointCloudLODManager::ProcessLOD(const TArray<FLidarPointCloudLODMan
 			const float RenderDataSmoothingMaxFrametime = GetDefault<ULidarPointCloudSettings>()->RenderDataSmoothingMaxFrametime;
 
 			// Iterate over proxies and, if valid, update their data
-			for (const FLidarPointCloudProxyUpdateData& UpdateData : ProxyUpdateData)
+			for (FLidarPointCloudProxyUpdateData& UpdateData : ProxyUpdateData)
 			{
 				// Check for proxy's validity, in case it has been destroyed since the update was issued
 				if (TSharedPtr<FLidarPointCloudSceneProxyWrapper, ESPMode::ThreadSafe> SceneProxyWrapper = UpdateData.SceneProxyWrapper.Pin())
 				{
 					if (SceneProxyWrapper->Proxy)
 					{
-						for (const FLidarPointCloudProxyUpdateDataNode& Node : UpdateData.SelectedNodes)
+						for (FLidarPointCloudProxyUpdateDataNode& Node : UpdateData.SelectedNodes)
 						{
-							if (Node.DataNode->BuildDataCache(UpdateData.bUseStaticBuffers))
+							if (Node.BuildDataCache(UpdateData.bUseStaticBuffers))
 							{
-								MaxPointsPerNode = FMath::Max(MaxPointsPerNode, Node.DataNode->GetNumVisiblePoints());
+								MaxPointsPerNode = FMath::Max(MaxPointsPerNode, (uint32)Node.NumVisiblePoints);
 							}
 
 							// Split building render data across multiple frames, to avoid stuttering
@@ -715,6 +717,14 @@ int64 FLidarPointCloudLODManager::ProcessLOD(const TArray<FLidarPointCloudLODMan
 	return NewNumPointsInFrustum;
 }
 
+void FLidarPointCloudLODManager::ForceProcessLOD()
+{
+	check(IsInGameThread());
+	
+	PrepareProxies();
+	ProcessLOD(RegisteredProxies, Time, LastPointBudget, GetClippingVolumes());
+}
+
 void FLidarPointCloudLODManager::PrepareProxies()
 {
 	FLidarPointCloudViewData ViewData(true);
@@ -761,6 +771,13 @@ void FLidarPointCloudLODManager::PrepareProxies()
 
 							// Recreate the Traversal Octree
 							RegisteredProxy.TraversalOctree = MakeShareable(new FLidarPointCloudTraversalOctree(&PointCloud->Octree, Component->GetComponentTransform()));
+
+							// If the recreation of the Traversal Octree was unsuccessful, skip further processing
+							if (!RegisteredProxy.TraversalOctree->bValid)
+							{
+								continue;
+							}
+							
 							RegisteredProxy.PointCloud->Octree.RegisterTraversalOctree(RegisteredProxy.TraversalOctree);
 						}
 
@@ -846,6 +863,12 @@ TArray<FLidarPointCloudClippingVolumeParams> FLidarPointCloudLODManager::GetClip
 	return ClippingVolumes;
 }
 
+FLidarPointCloudLODManager& FLidarPointCloudLODManager::Get()
+{
+	static FLidarPointCloudLODManager Instance;
+	return Instance;
+}
+
 FLidarPointCloudLODManager::FRegisteredProxy::FRegisteredProxy(TWeakObjectPtr<ULidarPointCloudComponent> Component, TWeakPtr<FLidarPointCloudSceneProxyWrapper, ESPMode::ThreadSafe> SceneProxyWrapper)
 	: Component(Component)
 	, PointCloud(Component->GetPointCloud())
diff --git a/Engine/Plugins/Enterprise/LidarPointCloud/Source/LidarPointCloudRuntime/Private/LidarPointCloudOctree.cpp b/Engine/Plugins/Enterprise/LidarPointCloud/Source/LidarPointCloudRuntime/Private/LidarPointCloudOctree.cpp
--- a/Engine/Plugins/Enterprise/LidarPointCloud/Source/LidarPointCloudRuntime/Private/LidarPointCloudOctree.cpp
+++ b/Engine/Plugins/Enterprise/LidarPointCloud/Source/LidarPointCloudRuntime/Private/LidarPointCloudOctree.cpp
@@ -136,13 +136,13 @@ FLidarPointCloudOctreeNode::FLidarPointCloudOctreeNode(FLidarPointCloudOctree* T
 
 FLidarPointCloudOctreeNode::~FLidarPointCloudOctreeNode()
 {
+	ReleaseDataCache();
+
 	for (int32 i = 0; i < Children.Num(); i++)
 	{
 		delete Children[i];
 		Children[i] = nullptr;
 	}
-
-	ReleaseDataCache();
 }
 
 void FLidarPointCloudOctreeNode::UpdateNumVisiblePoints()
@@ -188,47 +188,45 @@ FLidarPointCloudPoint* FLidarPointCloudOctreeNode::GetPersistentData() const
 }
 
 bool FLidarPointCloudOctreeNode::BuildDataCache(bool bUseStaticBuffers)
-{
+{	
 	// Only include nodes with available data
 	if (HasData() && GetNumVisiblePoints())
 	{
 		// Make sure to release the unnecessary buffer
 		if (bUseStaticBuffers)
 		{
-			if (DataCache)
+			if (DataCache.IsValid())
 			{
 				DataCache->ReleaseResource();
-				delete DataCache;
-				DataCache = nullptr;
+				DataCache.Reset();
 				bRenderDataDirty = true;
 			}
 
-			if (!VertexFactory)
+			if (!VertexFactory.IsValid())
 			{
-				VertexFactory = new FLidarPointCloudVertexFactory();
+				VertexFactory = MakeShareable(new FLidarPointCloudVertexFactory());
 				bRenderDataDirty = true;
 			}
 		}
 		else
 		{
-			if (VertexFactory)
+			if (VertexFactory.IsValid())
 			{
 				VertexFactory->ReleaseResource();
-				delete VertexFactory;
-				VertexFactory = nullptr;
+				VertexFactory.Reset();
 				bRenderDataDirty = true;
 			}
 
-			if (!DataCache)
+			if (!DataCache.IsValid())
 			{
-				DataCache = new FLidarPointCloudRenderBuffer();
+				DataCache = MakeShareable(new FLidarPointCloudRenderBuffer());
 				bRenderDataDirty = true;
 			}
 		}
 
 		if (bRenderDataDirty)
 		{
-			if (DataCache)
+			if (DataCache.IsValid())
 			{
 				DataCache->Resize(GetNumVisiblePoints() * 5);
 
@@ -241,7 +239,7 @@ bool FLidarPointCloudOctreeNode::BuildDataCache(bool bUseStaticBuffers)
 				RHIUnlockBuffer(DataCache->Buffer);
 			}
 
-			if (VertexFactory)
+			if (VertexFactory.IsValid())
 			{
 				VertexFactory->Initialize(GetData(), GetNumVisiblePoints());
 			}
@@ -813,25 +811,23 @@ void FLidarPointCloudOctreeNode::ReleaseData(bool bForce)
 
 void FLidarPointCloudOctreeNode::ReleaseDataCache()
 {
-	if (VertexFactory || DataCache)
+	if (VertexFactory.IsValid() || DataCache.IsValid())
 	{
 		ENQUEUE_RENDER_COMMAND(LidarPointCloudOctreeNode_ReleaseDataCache)([DataCache = DataCache, VertexFactory = VertexFactory](FRHICommandListImmediate& RHICmdList)
+		{
+			if (DataCache.IsValid())
 			{
-				if (DataCache)
-				{
-					DataCache->ReleaseResource();
-					delete DataCache;
-				}
+				DataCache->ReleaseResource();
+			}
 
-				if (VertexFactory)
-				{
-					VertexFactory->ReleaseResource();
-					delete VertexFactory;
-				}
-			});
+			if (VertexFactory.IsValid())
+			{
+				VertexFactory->ReleaseResource();
+			}
+		});
 
-		DataCache = nullptr;
-		VertexFactory = nullptr;
+		DataCache.Reset();
+		VertexFactory.Reset();
 	}
 }
 
@@ -855,7 +851,7 @@ void FLidarPointCloudOctreeNode::SortVisiblePoints()
 //////////////////////////////////////////////////////////// FLidarPointCloudOctree
 
 FLidarPointCloudOctree::FLidarPointCloudOctree(ULidarPointCloud* Owner)
-	: Root(nullptr, 0)
+	: Root(new FLidarPointCloudOctreeNode(nullptr, 0))
 	, Owner(Owner)
 	, BulkData(this)
 	, bStreamingBusy(false)
@@ -868,12 +864,12 @@ FLidarPointCloudOctree::FLidarPointCloudOctree(ULidarPointCloud* Owner)
 	// Account for the Root
 	NodeCount[0].Increment();
 
-	Root.Tree = this;
+	Root->Tree = this;
 }
 
 FLidarPointCloudOctree::~FLidarPointCloudOctree()
 {
-	MarkTraversalOctreesForInvalidation();
+	Empty(true);
 }
 
 int32 FLidarPointCloudOctree::GetNumLODs() const
@@ -1302,7 +1298,7 @@ void FLidarPointCloudOctree::GetPointsAsCopiesInBatches(TFunction<void(TSharedPt
 
 	TQueue<FLidarPointCloudOctreeNode*> Nodes;
 	FLidarPointCloudOctreeNode* CurrentNode = nullptr;
-	Nodes.Enqueue(&Root);
+	Nodes.Enqueue(Root);
 	while (Nodes.Dequeue(CurrentNode))
 	{
 		FOR_RO(Point, CurrentNode)
@@ -1784,7 +1780,7 @@ void FLidarPointCloudOctree::RemovePoint(const FLidarPointCloudPoint* Point)
 	}
 	
 	int32 Index = -1;
-	FLidarPointCloudOctreeNode* CurrentNode = &Root;
+	FLidarPointCloudOctreeNode* CurrentNode = Root;
 	while (CurrentNode)
 	{
 		const FLidarPointCloudPoint* RESTRICT Start = CurrentNode->GetData();
@@ -1816,7 +1812,7 @@ void FLidarPointCloudOctree::RemovePoint(const FLidarPointCloudPoint* Point)
 void FLidarPointCloudOctree::RemovePoint(FLidarPointCloudPoint Point)
 {
 	int32 Index = -1;
-	FLidarPointCloudOctreeNode* CurrentNode = &Root;
+	FLidarPointCloudOctreeNode* CurrentNode = Root;
 	while (CurrentNode)
 	{
 		const FLidarPointCloudPoint* RESTRICT Start = CurrentNode->GetData();
@@ -2039,24 +2035,31 @@ void FLidarPointCloudOctree::CalculateNormals(FThreadSafeBool* bCancelled, int32
 
 void FLidarPointCloudOctree::Empty(bool bDestroyNodes)
 {
+	FScopeLock OctreeLock(&DataLock);
+	MarkTraversalOctreesForInvalidation();
+
 	if (bDestroyNodes)
 	{
-		Root.~FLidarPointCloudOctreeNode();
-
 		// Reset node counters
 		for (FThreadSafeCounter& Count : NodeCount)
 		{
 			Count.Reset();
 		}
-
-		new (&Root) FLidarPointCloudOctreeNode(this, 0);
+		
+		// Perform the actual deletion on the RT, so there is no concurrent access issues
+		ENQUEUE_RENDER_COMMAND(LidarPointCloudOctreeNode_DeleteOldData)([OldRoot = Root](FRHICommandListImmediate& RHICmdList)
+		{
+			delete OldRoot;
+		});
+		
+		Root = new FLidarPointCloudOctreeNode(this, 0);
 
 		QueuedNodes.Empty();
 		NodesInUse.Reset();
 	}
 	else
 	{
-		Root.Empty(true);
+		Root->Empty(true);
 	}
 	
 	// Reset point counters
@@ -2064,8 +2067,6 @@ void FLidarPointCloudOctree::Empty(bool bDestroyNodes)
 	{
 		Count.Reset();
 	}
-
-	MarkTraversalOctreesForInvalidation();
 }
 
 void FLidarPointCloudOctree::UnregisterTraversalOctree(FLidarPointCloudTraversalOctree* TraversalOctree)
@@ -2226,7 +2227,7 @@ void FLidarPointCloudOctree::OptimizeForStaticData()
 	if (IsOptimizedForDynamicData())
 	{
 		// Move data out of the root node
-		TArray<FLidarPointCloudPoint> SourceData = MoveTemp(Root.Data);
+		TArray<FLidarPointCloudPoint> SourceData = MoveTemp(Root->Data);
 		
 		// Destroy current tree structure
 		Empty(true);
@@ -2254,14 +2255,14 @@ void FLidarPointCloudOctree::RefreshAllocatedSize()
 	PreviousAllocatedStructureSize += SharedData.GetAllocatedSize();
 	PreviousAllocatedStructureSize += PointCount.GetAllocatedSize();
 
-	PreviousAllocatedSize = PreviousAllocatedStructureSize + Root.GetAllocatedSize(true, true);
-	PreviousAllocatedStructureSize += Root.GetAllocatedSize(true, false);
+	PreviousAllocatedSize = PreviousAllocatedStructureSize + Root->GetAllocatedSize(true, true);
+	PreviousAllocatedStructureSize += Root->GetAllocatedSize(true, false);
 }
 
 template <typename T>
 void FLidarPointCloudOctree::InsertPoints_Internal(T Points, const int64& Count, ELidarPointCloudDuplicateHandling DuplicateHandling, bool bRefreshPointsBounds, const FVector& Translation)
 {
-	Root.InsertPoints(Points, Count, DuplicateHandling, Translation);
+	Root->InsertPoints(Points, Count, DuplicateHandling, Translation);
 	MarkTraversalOctreesForInvalidation();
 	if (bRefreshPointsBounds)
 	{
@@ -2349,7 +2350,7 @@ void FLidarPointCloudOctree::Serialize(FArchive& Ar)
 		TArray<FLidarPointCloudPoint_Legacy> DataArray;
 
 		TArray<FLidarPointCloudOctreeNode*> Nodes;
-		Nodes.Add(&Root);
+		Nodes.Add(Root);
 		while (Nodes.Num())
 		{
 			FLidarPointCloudOctreeNode* CurrentNode = Nodes.Pop(false);
@@ -2722,9 +2723,9 @@ FLidarPointCloudTraversalOctree::FLidarPointCloudTraversalOctree(FLidarPointClou
 	}
 
 	// Star cloning the node data
-	Root.Build(this, &Octree->Root, LocalToWorld, Octree->Owner->GetLocationOffset().ToVector());
+	Root.Build(this, Octree->Root, LocalToWorld, Octree->Owner->GetLocationOffset().ToVector());
 
-	bValid = true;
+	bValid = NumLODs > 0;
 }
 
 void FLidarPointCloudTraversalOctree::CalculateVisibilityStructure(TArray<uint32>& OutData)
diff --git a/Engine/Plugins/Enterprise/LidarPointCloud/Source/LidarPointCloudRuntime/Private/LidarPointCloudOctreeMacros.h b/Engine/Plugins/Enterprise/LidarPointCloud/Source/LidarPointCloudRuntime/Private/LidarPointCloudOctreeMacros.h
--- a/Engine/Plugins/Enterprise/LidarPointCloud/Source/LidarPointCloudRuntime/Private/LidarPointCloudOctreeMacros.h
+++ b/Engine/Plugins/Enterprise/LidarPointCloud/Source/LidarPointCloudRuntime/Private/LidarPointCloudOctreeMacros.h
@@ -14,7 +14,7 @@
 {\
 	TQueue<Const FLidarPointCloudOctreeNode*> Nodes;\
 	Const FLidarPointCloudOctreeNode* CurrentNode = nullptr;\
-	Nodes.Enqueue(&Root);\
+	Nodes.Enqueue(Root);\
 	while (Nodes.Dequeue(CurrentNode))\
 	{\
 		{ Action } \
diff --git a/Engine/Plugins/Enterprise/LidarPointCloud/Source/LidarPointCloudRuntime/Private/Meshing/LidarPointCloudMeshing.cpp b/Engine/Plugins/Enterprise/LidarPointCloud/Source/LidarPointCloudRuntime/Private/Meshing/LidarPointCloudMeshing.cpp
--- a/Engine/Plugins/Enterprise/LidarPointCloud/Source/LidarPointCloudRuntime/Private/Meshing/LidarPointCloudMeshing.cpp
+++ b/Engine/Plugins/Enterprise/LidarPointCloud/Source/LidarPointCloudRuntime/Private/Meshing/LidarPointCloudMeshing.cpp
@@ -42,7 +42,7 @@ void LidarPointCloudMeshing::CalculateNormals(FLidarPointCloudOctree* Octree, FT
 
 	TQueue<FSamplingUnit*> Q;
 	{
-		FSamplingUnit* Root = new FSamplingUnit(FVector::ZeroVector, Octree->SharedData[0].Extent, &Octree->Root);
+		FSamplingUnit* Root = new FSamplingUnit(FVector::ZeroVector, Octree->SharedData[0].Extent, Octree->Root);
 		if (InPointSelection.Num() == 0)
 		{
 			Octree->GetPoints(Root->Points);
diff --git a/Engine/Plugins/Enterprise/LidarPointCloud/Source/LidarPointCloudRuntime/Private/Rendering/LidarPointCloudRendering.cpp b/Engine/Plugins/Enterprise/LidarPointCloud/Source/LidarPointCloudRuntime/Private/Rendering/LidarPointCloudRendering.cpp
--- a/Engine/Plugins/Enterprise/LidarPointCloud/Source/LidarPointCloudRuntime/Private/Rendering/LidarPointCloudRendering.cpp
+++ b/Engine/Plugins/Enterprise/LidarPointCloud/Source/LidarPointCloudRuntime/Private/Rendering/LidarPointCloudRendering.cpp
@@ -16,6 +16,19 @@
 
 DECLARE_DWORD_COUNTER_STAT(TEXT("Draw Calls"), STAT_DrawCallCount, STATGROUP_LidarPointCloud)
 
+bool FLidarPointCloudProxyUpdateDataNode::BuildDataCache(bool bUseStaticBuffers)
+{
+	if(DataNode && DataNode->BuildDataCache(bUseStaticBuffers))
+	{
+		VertexFactory = DataNode->GetVertexFactory();
+		DataCache = DataNode->GetDataCache();
+		DataNode = nullptr;
+		return  true;
+	}
+	
+	return false;
+}
+
 FLidarPointCloudProxyUpdateData::FLidarPointCloudProxyUpdateData()
 	: NumElements(0)
 	, VDMultiplier(1)
@@ -205,13 +218,13 @@ public:
 
 					for (const FLidarPointCloudProxyUpdateDataNode& Node : RenderData.SelectedNodes)
 					{
-						if (Node.DataNode && ((RenderData.bUseStaticBuffers && Node.DataNode->GetVertexFactory()) || (!RenderData.bUseStaticBuffers && Node.DataNode->GetDataCache())))
+						if ((RenderData.bUseStaticBuffers && Node.VertexFactory.IsValid() && Node.VertexFactory->IsInitialized()) || (!RenderData.bUseStaticBuffers && Node.DataCache.IsValid()))
 						{
 							FMeshBatch& MeshBatch = Collector.AllocateMesh();
 
 							MeshBatch.Type = bUsesSprites ? PT_TriangleList : PT_PointList;
 							MeshBatch.LODIndex = 0;
-							MeshBatch.VertexFactory = RenderData.bUseStaticBuffers ? Node.DataNode->GetVertexFactory() : (FVertexFactory*)&GLidarPointCloudSharedVertexFactory;
+							MeshBatch.VertexFactory = RenderData.bUseStaticBuffers ? Node.VertexFactory.Get() : (FVertexFactory*)&GLidarPointCloudSharedVertexFactory;
 							MeshBatch.bWireframe = false;
 							MeshBatch.MaterialRenderProxy = RenderData.RenderParams.Material->GetRenderProxy();
 							MeshBatch.ReverseCulling = IsLocalToWorldDeterminantNegative();
@@ -348,7 +361,7 @@ public:
 			UserDataElement.bStartClipped |= ClippingVolume.Mode == ELidarClippingVolumeMode::ClipOutside;
 		}
 
-		UserDataElement.DataBuffer = Node.DataNode->GetDataCache() ? Node.DataNode->GetDataCache()->SRV : nullptr;
+		UserDataElement.DataBuffer = Node.DataCache ? Node.DataCache->SRV : nullptr;
 		UserDataElement.TreeBuffer = TreeBuffer->SRV;
 
 		return UserDataElement;
diff --git a/Engine/Plugins/Enterprise/LidarPointCloud/Source/LidarPointCloudRuntime/Public/LidarPointCloud.h b/Engine/Plugins/Enterprise/LidarPointCloud/Source/LidarPointCloudRuntime/Public/LidarPointCloud.h
--- a/Engine/Plugins/Enterprise/LidarPointCloud/Source/LidarPointCloudRuntime/Public/LidarPointCloud.h
+++ b/Engine/Plugins/Enterprise/LidarPointCloud/Source/LidarPointCloudRuntime/Public/LidarPointCloud.h
@@ -152,6 +152,8 @@ private:
 	/** Description of collision */
 	UPROPERTY(transient, duplicatetransient)
 	UBodySetup* BodySetup;
+	UPROPERTY(transient, duplicatetransient)
+	UBodySetup* NewBodySetup;
 
 	/** Used for collision building */
 	FThreadSafeBool bCollisionBuildInProgress;
@@ -175,7 +177,7 @@ public:
 	virtual void PreSave(FObjectPreSaveContext ObjectSaveContext) override;
 #if WITH_EDITOR
 	virtual void ClearAllCachedCookedPlatformData() override;
-	virtual void PostEditChangeProperty(FPropertyChangedEvent & PropertyChangedEvent) override;
+	virtual void PostEditChangeProperty(FPropertyChangedEvent& PropertyChangedEvent) override;
 #endif
 	// End UObject Interface.
 
@@ -644,8 +646,8 @@ public:
 
 private:
 	/** Once async physics cook is done, create needed state */
-	void FinishPhysicsAsyncCook(UBodySetup* NewBodySetup, TSharedRef<FLidarPointCloudNotification, ESPMode::ThreadSafe> Notification) { FinishPhysicsAsyncCook(true, NewBodySetup, Notification); }
-	void FinishPhysicsAsyncCook(bool bSuccess, UBodySetup* NewBodySetup, TSharedRef<FLidarPointCloudNotification, ESPMode::ThreadSafe> Notification);
+	void FinishPhysicsAsyncCook(TSharedRef<FLidarPointCloudNotification, ESPMode::ThreadSafe> Notification) { FinishPhysicsAsyncCook(true, Notification); }
+	void FinishPhysicsAsyncCook(bool bSuccess, TSharedRef<FLidarPointCloudNotification, ESPMode::ThreadSafe> Notification);
 
 	void InitializeCollisionRendering();
 	void ReleaseCollisionRendering(bool bDestroyAfterRelease);
diff --git a/Engine/Plugins/Enterprise/LidarPointCloud/Source/LidarPointCloudRuntime/Public/LidarPointCloudLODManager.h b/Engine/Plugins/Enterprise/LidarPointCloud/Source/LidarPointCloudRuntime/Public/LidarPointCloudLODManager.h
--- a/Engine/Plugins/Enterprise/LidarPointCloud/Source/LidarPointCloudRuntime/Public/LidarPointCloudLODManager.h
+++ b/Engine/Plugins/Enterprise/LidarPointCloud/Source/LidarPointCloudRuntime/Public/LidarPointCloudLODManager.h
@@ -70,6 +70,9 @@ class FLidarPointCloudLODManager : public FTickableGameObject
 	/** Stores the number of points in visible frustum during last frame */
 	FThreadSafeCounter64 NumPointsInFrustum;
 
+	/** Stores the last calculated point budget */
+	uint32 LastPointBudget;
+
 public:
 	FLidarPointCloudLODManager();
 
@@ -83,6 +86,8 @@ public:
 
 	static void RegisterProxy(ULidarPointCloudComponent* Component, TWeakPtr<FLidarPointCloudSceneProxyWrapper, ESPMode::ThreadSafe> SceneProxyWrapper);
 
+	static void RefreshLOD();
+
 private:
 	/**
 	 * This function:
@@ -94,9 +99,14 @@ private:
 	 */
 	int64 ProcessLOD(const TArray<FRegisteredProxy>& RegisteredProxies, const float CurrentTime, const uint32 PointBudget, const TArray<FLidarPointCloudClippingVolumeParams>& ClippingVolumes);
 
+	/** Forces a refresh of the LOD processing. Must be called from a GT */
+	void ForceProcessLOD();
+
 	/** Called to prepare the proxies for processing */
 	void PrepareProxies();
 
 	/** Compiles a list of all clipping volumes affecting any of the registered proxies */
 	TArray<FLidarPointCloudClippingVolumeParams> GetClippingVolumes() const;
+
+	static FLidarPointCloudLODManager& Get();
 };
\ No newline at end of file
diff --git a/Engine/Plugins/Enterprise/LidarPointCloud/Source/LidarPointCloudRuntime/Public/LidarPointCloudOctree.h b/Engine/Plugins/Enterprise/LidarPointCloud/Source/LidarPointCloudRuntime/Public/LidarPointCloudOctree.h
--- a/Engine/Plugins/Enterprise/LidarPointCloud/Source/LidarPointCloudRuntime/Public/LidarPointCloudOctree.h
+++ b/Engine/Plugins/Enterprise/LidarPointCloud/Source/LidarPointCloudRuntime/Public/LidarPointCloudOctree.h
@@ -18,6 +18,9 @@ class FLidarPointCloudOctree;
 struct FLidarPointCloudTraversalOctree;
 struct FLidarPointCloudTraversalOctreeNode;
 
+class FLidarPointCloudRenderBuffer;
+class FLidarPointCloudVertexFactory;
+
 /**
  * WARNING: Exercise caution when modifying the contents of the Octree, as it may be in use by the Rendering Thread via FPointCloudSceneProxy
  * Use the FLidarPointCloudOctree::DataLock prior to such attempt
@@ -92,8 +95,8 @@ private:
 	uint32 BulkDataSize;
 
 	/** Holds render data for this node */
-	class FLidarPointCloudRenderBuffer* DataCache;
-	class FLidarPointCloudVertexFactory* VertexFactory;
+	TSharedPtr<FLidarPointCloudRenderBuffer> DataCache;
+	TSharedPtr<FLidarPointCloudVertexFactory> VertexFactory;
 
 	bool bRenderDataDirty;
 
@@ -119,10 +122,10 @@ public:
 	FLidarPointCloudPoint* GetPersistentData() const;
 
 	/** Returns a pointer to the point data */
-	FORCEINLINE FLidarPointCloudRenderBuffer* GetDataCache() { return DataCache; }
+	FORCEINLINE TSharedPtr<FLidarPointCloudRenderBuffer> GetDataCache() { return DataCache; }
 
 	/** Return a pointer to the vertex factory containing pre-cached geometry */
-	FORCEINLINE FLidarPointCloudVertexFactory* GetVertexFactory() { return VertexFactory; }
+	FORCEINLINE TSharedPtr<FLidarPointCloudVertexFactory> GetVertexFactory() { return VertexFactory; }
 
 	/**
 	 * Builds and updates the necessary render data buffers
@@ -166,7 +169,7 @@ public:
 	bool HasData() const { return bHasData; }
 
 	/** Returns true, if the node has its data loaded */
-	bool HasRenderData() const { return DataCache || VertexFactory; }
+	bool HasRenderData() const { return DataCache.IsValid() || VertexFactory.IsValid(); }
 
 	/**
 	 * Releases the BulkData
@@ -235,7 +238,7 @@ public:
 	FCriticalSection DataReleaseLock;
 
 private:
-	FLidarPointCloudOctreeNode Root;
+	FLidarPointCloudOctreeNode* Root;
 	
 	/** Stores shared per-LOD node data. */
 	TArray<FSharedLODData> SharedData;
@@ -301,7 +304,7 @@ public:
 	FLidarPointCloudOctree& operator=(FLidarPointCloudOctree&&) = delete;
 
 	/** Returns true if the Root node exists and has any data assigned. */
-	bool HasData() const { return Root.GetNumPoints() > 0; }
+	bool HasData() const { return Root->GetNumPoints() > 0; }
 
 	/** Returns the number of different LODs. */
 	int32 GetNumLODs() const;
diff --git a/Engine/Plugins/Enterprise/LidarPointCloud/Source/LidarPointCloudRuntime/Public/Rendering/LidarPointCloudRendering.h b/Engine/Plugins/Enterprise/LidarPointCloud/Source/LidarPointCloudRuntime/Public/Rendering/LidarPointCloudRendering.h
--- a/Engine/Plugins/Enterprise/LidarPointCloud/Source/LidarPointCloudRuntime/Public/Rendering/LidarPointCloudRendering.h
+++ b/Engine/Plugins/Enterprise/LidarPointCloud/Source/LidarPointCloudRuntime/Public/Rendering/LidarPointCloudRendering.h
@@ -16,17 +16,32 @@ struct FLidarPointCloudSceneProxyWrapper
 
 struct FLidarPointCloudProxyUpdateDataNode
 {
+private:
+	struct FLidarPointCloudOctreeNode* DataNode;
+	
+public:
 	uint8 VirtualDepth;
 	int64 NumVisiblePoints;
-	struct FLidarPointCloudOctreeNode* DataNode;
+
+	/**
+	 * Holds render data pointers for this node
+	 * This allows us to avoid exposing potentially invalid DataNode pointer in Render Thread
+	 */
+	TSharedPtr<class FLidarPointCloudRenderBuffer> DataCache;
+	TSharedPtr<class FLidarPointCloudVertexFactory> VertexFactory;
 
 	FLidarPointCloudProxyUpdateDataNode() : FLidarPointCloudProxyUpdateDataNode(0, 0, nullptr) {}
 	FLidarPointCloudProxyUpdateDataNode(uint8 VirtualDepth, int64 NumVisiblePoints, FLidarPointCloudOctreeNode* DataNode)
-		: VirtualDepth(VirtualDepth)
+		: DataNode(DataNode)
+		, VirtualDepth(VirtualDepth)
 		, NumVisiblePoints(NumVisiblePoints)
-		, DataNode(DataNode)
+		, DataCache(nullptr)
+		, VertexFactory(nullptr)
 	{
 	}
+
+	/** Passthrough method */
+	bool BuildDataCache(bool bUseStaticBuffers);
 };
 
 /** Used to pass data to RT to update the proxy's render data */
