Fix PVS-Studio static analysis warnings

diff --git a/Engine/Plugins/Enterprise/LidarPointCloud/Source/LidarPointCloudRuntime/Private/LidarPointCloudLODManager.cpp b/Engine/Plugins/Enterprise/LidarPointCloud/Source/LidarPointCloudRuntime/Private/LidarPointCloudLODManager.cpp
--- a/Engine/Plugins/Enterprise/LidarPointCloud/Source/LidarPointCloudRuntime/Private/LidarPointCloudLODManager.cpp
+++ b/Engine/Plugins/Enterprise/LidarPointCloud/Source/LidarPointCloudRuntime/Private/LidarPointCloudLODManager.cpp
@@ -98,7 +98,8 @@ void FLidarPointCloudViewData::Compute()
 		if (LP && LP->ViewportClient)
 		{
 			FSceneViewProjectionData ProjectionData;
-			if (LP->GetProjectionData(LP->ViewportClient->Viewport, ProjectionData, GEngine->IsStereoscopic3D() ? EStereoscopicEye::eSSE_LEFT_EYE : INDEX_NONE))
+			if (LP->GetProjectionData(LP->ViewportClient->Viewport, ProjectionData,
+				GEngine->IsStereoscopic3D() ? EStereoscopicEye::eSSE_LEFT_EYE : EStereoscopicEye::eSSE_MONOSCOPIC))
 			{
 				ViewOrigin = ProjectionData.ViewOrigin;
 				FMatrix ViewRotationMatrix = ProjectionData.ViewRotationMatrix;
diff --git a/Engine/Source/Runtime/VulkanRHI/Private/VulkanUAV.cpp b/Engine/Source/Runtime/VulkanRHI/Private/VulkanUAV.cpp
--- a/Engine/Source/Runtime/VulkanRHI/Private/VulkanUAV.cpp
+++ b/Engine/Source/Runtime/VulkanRHI/Private/VulkanUAV.cpp
@@ -92,7 +92,7 @@ FVulkanShaderResourceView::~FVulkanShaderResourceView()
 void FVulkanShaderResourceView::Clear()
 {
 #if VULKAN_RHI_RAYTRACING
-	if (AccelerationStructureHandle)
+	if (Device && AccelerationStructureHandle)
 	{
 		Device->GetDeferredDeletionQueue().EnqueueResource(VulkanRHI::FDeferredDeletionQueue2::EType::AccelerationStructure, AccelerationStructureHandle);
 	}
