LidarPointCloud Rendering: Create a dummy buffer to bind when there is no Databuffer to avoid empty descriptor write in Vulkan (fixes crash in Linux when importing point cloud asset).
Tested on Linux with Nvidia and AMD.

diff --git a/Engine/Plugins/Enterprise/LidarPointCloud/Source/LidarPointCloudRuntime/Private/Rendering/LidarPointCloudRenderBuffers.cpp b/Engine/Plugins/Enterprise/LidarPointCloud/Source/LidarPointCloudRuntime/Private/Rendering/LidarPointCloudRenderBuffers.cpp
--- a/Engine/Plugins/Enterprise/LidarPointCloud/Source/LidarPointCloudRuntime/Private/Rendering/LidarPointCloudRenderBuffers.cpp
+++ b/Engine/Plugins/Enterprise/LidarPointCloud/Source/LidarPointCloudRuntime/Private/Rendering/LidarPointCloudRenderBuffers.cpp
@@ -21,6 +21,7 @@
 
 TGlobalResource<FLidarPointCloudIndexBuffer> GLidarPointCloudIndexBuffer;
 TGlobalResource<FLidarPointCloudSharedVertexFactory> GLidarPointCloudSharedVertexFactory;
+TGlobalResource<FLidarPointCloudRenderBuffer> GDummyLidarPointCloudRenderBuffer(4);
 
 //////////////////////////////////////////////////////////// Index Buffer
 
@@ -182,6 +183,12 @@ void FLidarPointCloudVertexFactoryShaderParameters::GetElementShaderBindings(con
 
 	SETSRVPARAM(TreeBuffer);
 	SETSRVPARAM(DataBuffer);
+	if (!UserData->DataBuffer && DataBuffer.IsBound())
+	{
+		FRHIShaderResourceView* Ptr = GDummyLidarPointCloudRenderBuffer.SRV;
+		ShaderBindings.Add(DataBuffer, Ptr);
+	}
+
 	SETPARAM(bEditorView);
 	SETPARAM(SelectionColor);
 	SETPARAM(LocationOffset);
diff --git a/Engine/Plugins/Enterprise/LidarPointCloud/Source/LidarPointCloudRuntime/Private/Rendering/LidarPointCloudRenderBuffers.h b/Engine/Plugins/Enterprise/LidarPointCloud/Source/LidarPointCloudRuntime/Private/Rendering/LidarPointCloudRenderBuffers.h
--- a/Engine/Plugins/Enterprise/LidarPointCloud/Source/LidarPointCloudRuntime/Private/Rendering/LidarPointCloudRenderBuffers.h
+++ b/Engine/Plugins/Enterprise/LidarPointCloud/Source/LidarPointCloudRuntime/Private/Rendering/LidarPointCloudRenderBuffers.h
@@ -47,6 +47,9 @@ public:
 class FLidarPointCloudRenderBuffer : public FLidarPointCloudBuffer, public FRenderResource
 {
 public:
+	FLidarPointCloudRenderBuffer() : FLidarPointCloudBuffer() {}
+	FLidarPointCloudRenderBuffer(uint32 Capacity) : FLidarPointCloudBuffer(Capacity) {}
+
 	FBufferRHIRef Buffer;
 	FShaderResourceViewRHIRef SRV;
 
@@ -199,4 +202,5 @@ class FLidarPointCloudSharedVertexFactory : public FLidarPointCloudVertexFactory
 
 /** A set of global render resources shared between all Lidar Point Cloud proxies */
 extern TGlobalResource<FLidarPointCloudIndexBuffer> GLidarPointCloudIndexBuffer;
-extern TGlobalResource<FLidarPointCloudSharedVertexFactory> GLidarPointCloudSharedVertexFactory;
\ No newline at end of file
+extern TGlobalResource<FLidarPointCloudSharedVertexFactory> GLidarPointCloudSharedVertexFactory;
+extern TGlobalResource<FLidarPointCloudRenderBuffer> GDummyLidarPointCloudRenderBuffer;
