Fixed Editor fails to cook level with LiDAR Point Cloud Actor imported

diff --git a/Engine/Plugins/Enterprise/LidarPointCloud/Source/LidarPointCloudRuntime/Private/LidarPointCloud.cpp b/Engine/Plugins/Enterprise/LidarPointCloud/Source/LidarPointCloudRuntime/Private/LidarPointCloud.cpp
--- a/Engine/Plugins/Enterprise/LidarPointCloud/Source/LidarPointCloudRuntime/Private/LidarPointCloud.cpp
+++ b/Engine/Plugins/Enterprise/LidarPointCloud/Source/LidarPointCloudRuntime/Private/LidarPointCloud.cpp
@@ -379,7 +379,7 @@ void ULidarPointCloud::ClearAllCachedCookedPlatformData()
 {
 	if (GetDefault<ULidarPointCloudSettings>()->bReleaseAssetAfterCooking)
 	{
-		UPackageTools::ReloadPackages({ GetOutermost() });
+		Octree.ReleaseAllNodes(true);
 	}
 }
 
diff --git a/Engine/Plugins/Enterprise/LidarPointCloud/Source/LidarPointCloudRuntime/Private/LidarPointCloudLODManager.cpp b/Engine/Plugins/Enterprise/LidarPointCloud/Source/LidarPointCloudRuntime/Private/LidarPointCloudLODManager.cpp
--- a/Engine/Plugins/Enterprise/LidarPointCloud/Source/LidarPointCloudRuntime/Private/LidarPointCloudLODManager.cpp
+++ b/Engine/Plugins/Enterprise/LidarPointCloud/Source/LidarPointCloudRuntime/Private/LidarPointCloudLODManager.cpp
@@ -720,10 +720,11 @@ int64 FLidarPointCloudLODManager::ProcessLOD(const TArray<FLidarPointCloudLODMan
 
 void FLidarPointCloudLODManager::ForceProcessLOD()
 {
-	check(IsInGameThread());
-	
-	PrepareProxies();
-	ProcessLOD(RegisteredProxies, Time, LastPointBudget, GetClippingVolumes());
+	if(IsInGameThread())
+	{
+		PrepareProxies();
+		ProcessLOD(RegisteredProxies, Time, LastPointBudget, GetClippingVolumes());
+	}
 }
 
 void FLidarPointCloudLODManager::PrepareProxies()
diff --git a/Engine/Plugins/Enterprise/LidarPointCloud/Source/LidarPointCloudRuntime/Private/LidarPointCloudOctree.cpp b/Engine/Plugins/Enterprise/LidarPointCloud/Source/LidarPointCloudRuntime/Private/LidarPointCloudOctree.cpp
--- a/Engine/Plugins/Enterprise/LidarPointCloud/Source/LidarPointCloudRuntime/Private/LidarPointCloudOctree.cpp
+++ b/Engine/Plugins/Enterprise/LidarPointCloud/Source/LidarPointCloudRuntime/Private/LidarPointCloudOctree.cpp
@@ -853,7 +853,9 @@ void FLidarPointCloudOctreeNode::SortVisiblePoints()
 FLidarPointCloudOctree::FLidarPointCloudOctree(ULidarPointCloud* Owner)
 	: Root(new FLidarPointCloudOctreeNode(nullptr, 0))
 	, Owner(Owner)
-	, BulkData(this)
+#if WITH_EDITOR
+	, SavingBulkData(this)
+#endif
 	, bStreamingBusy(false)
 	, bIsFullyLoaded(false)
 {
@@ -2445,7 +2447,7 @@ void FLidarPointCloudOctree::Serialize(FArchive& Ar)
 		if (Ar.IsSaving())
 		{
 			LoadAllNodes(true);
-			BulkData.CloseReadHandle();
+			CloseReadHandle();
 		}
 
 		int64 BulkDataOffset = 0;
@@ -2477,8 +2479,17 @@ void FLidarPointCloudOctree::Serialize(FArchive& Ar)
 				}
 			}
 		}, true);
-
-		BulkData.Serialize(Ar, Owner);
+		
+#if WITH_EDITOR
+		if(Ar.IsSaving())
+		{
+			SavingBulkData.Serialize(Ar);
+		}
+		else
+#endif
+		{
+			BulkData.Serialize(Ar, Owner);
+		}
 	}
 
 	// Legacy Points Extent
@@ -2512,81 +2523,73 @@ void FLidarPointCloudOctree::Serialize(FArchive& Ar)
 void FLidarPointCloudOctree::SerializeBulkData(FArchive& Ar)
 {
 	if (Ar.IsSaving())
-	{
-		ITERATE_NODES({ Ar.Serialize(CurrentNode->Data.GetData(), CurrentNode->BulkDataSize); }, true);
-	}
-	else
 	{
 		ITERATE_NODES({
-			CurrentNode->Data.SetNumUninitialized(CurrentNode->NumPoints);
-			CurrentNode->bCanReleaseData = false;
 			Ar.Serialize(CurrentNode->Data.GetData(), CurrentNode->BulkDataSize);
-			CurrentNode->bHasData = true;
-			CurrentNode->bRenderDataDirty = true;
+			CurrentNode->ReleaseData(true);
 		}, true);
 	}
 }
 
-void FLidarPointCloudOctree::StreamNodeData(FLidarPointCloudOctreeNode* Node)
+IAsyncReadFileHandle* FLidarPointCloudOctree::GetReadHandle()
 {
-	Node->Data.SetNumUninitialized(Node->NumPoints);
-	Node->bHasData = BulkData.ReadRequest(Node->BulkDataOffset, Node->BulkDataSize, (uint8*)Node->Data.GetData());
-	
-	if (!Node->bHasData)
+	if(!ReadHandle)
 	{
-		Node->Data.Empty();
+		ReadHandle = BulkData.OpenAsyncReadHandle();
 	}
-}
-
-//////////////////////////////////////////////////////////// FLidarPointCloudOctree::FLidarPointCloudBulkData
-
-FLidarPointCloudOctree::FLidarPointCloudBulkData::FLidarPointCloudBulkData(FLidarPointCloudOctree* Octree)
-	: Octree(Octree)
-{
-	SetBulkDataFlags(BULKDATA_Force_NOT_InlinePayload | BULKDATA_Size64Bit);
-
-	// Dummy
-	Lock(LOCK_READ_WRITE);
-	*(uint8*)Realloc(1) = 0;
-	Unlock();
-}
 
-FLidarPointCloudOctree::FLidarPointCloudBulkData::~FLidarPointCloudBulkData()
-{
-	CloseReadHandle();
+	return ReadHandle;
 }
 
-bool FLidarPointCloudOctree::FLidarPointCloudBulkData::ReadRequest(int64 Offset, int64 BytesToRead, uint8* UserSuppliedMemory)
+void FLidarPointCloudOctree::CloseReadHandle()
 {
-	if (!ReadHandle)
+	if (ReadHandle)
 	{
-		ReadHandle = OpenAsyncReadHandle();
+		delete ReadHandle;
+		ReadHandle = nullptr;
 	}
+}
 
-	if (ReadHandle)
+void FLidarPointCloudOctree::StreamNodeData(FLidarPointCloudOctreeNode* Node)
+{
+	if (GetReadHandle())
 	{
-		IAsyncReadRequest* ReadRequest = ReadHandle->ReadRequest(GetBulkDataOffsetInFile() + Offset, BytesToRead, AIOP_Normal, nullptr, UserSuppliedMemory);
+		Node->Data.SetNumUninitialized(Node->NumPoints);		
+		IAsyncReadRequest* ReadRequest = ReadHandle->ReadRequest(
+			BulkData.GetBulkDataOffsetInFile() + Node->BulkDataOffset,
+			Node->BulkDataSize,
+			AIOP_Normal,
+			nullptr,
+			(uint8*)Node->Data.GetData());
 		ReadRequest->WaitCompletion();
 		delete ReadRequest;
-		return true;
+		Node->bHasData = true;
+	}
+	else
+	{
+		Node->bHasData = false;
+		Node->Data.Empty();
 	}
-
-	return false;
 }
 
-void FLidarPointCloudOctree::FLidarPointCloudBulkData::CloseReadHandle()
+//////////////////////////////////////////////////////////// FLidarPointCloudBulkData
+
+#if WITH_EDITOR
+FLidarPointCloudOctree::FLidarPointCloudBulkData::FLidarPointCloudBulkData(FLidarPointCloudOctree* Octree): Octree(Octree)
 {
-	if (ReadHandle)
-	{
-		delete ReadHandle;
-		ReadHandle = nullptr;
-	}
+	SetBulkDataFlags(BULKDATA_Force_NOT_InlinePayload);
+	Lock(LOCK_READ_WRITE);
+	Realloc(2);
+	Unlock();
 }
 
-void FLidarPointCloudOctree::FLidarPointCloudBulkData::SerializeElements(FArchive& Ar, void* Data)
+void FLidarPointCloudOctree::FLidarPointCloudBulkData::Serialize(FArchive& Ar)
 {
-	Octree->SerializeBulkData(Ar);
+	ElementSize = Ar.IsPersistent() && !Ar.IsObjectReferenceCollector() && !Ar.ShouldSkipBulkData() ? INT32_MAX : 1;
+	FByteBulkData::Serialize(Ar, Octree->GetOwner());
+	ElementSize = 1;
 }
+#endif
 
 //////////////////////////////////////////////////////////// FLidarPointCloudTraversalOctreeNode
 
diff --git a/Engine/Plugins/Enterprise/LidarPointCloud/Source/LidarPointCloudRuntime/Public/LidarPointCloudOctree.h b/Engine/Plugins/Enterprise/LidarPointCloud/Source/LidarPointCloudRuntime/Public/LidarPointCloudOctree.h
--- a/Engine/Plugins/Enterprise/LidarPointCloud/Source/LidarPointCloudRuntime/Public/LidarPointCloudOctree.h
+++ b/Engine/Plugins/Enterprise/LidarPointCloud/Source/LidarPointCloudRuntime/Public/LidarPointCloudOctree.h
@@ -267,25 +267,24 @@ private:
 	/** Pointer to the owner of this Octree */
 	ULidarPointCloud* Owner;
 
-	struct FLidarPointCloudBulkData : public FUntypedBulkData
+	IAsyncReadFileHandle* ReadHandle;
+
+	FByteBulkData BulkData;
+
+#if WITH_EDITOR
+	struct FLidarPointCloudBulkData : public FByteBulkData
 	{
-	private:
 		FLidarPointCloudOctree* Octree;
-		IAsyncReadFileHandle* ReadHandle;
-
-	public:
+		int32 ElementSize = 1;
+		
 		FLidarPointCloudBulkData(FLidarPointCloudOctree* Octree);
-		virtual ~FLidarPointCloudBulkData();
-		virtual int32 GetElementSize() const override { return 1; }
-		bool ReadRequest(int64 Offset, int64 BytesToRead, uint8* UserSuppliedMemory);
-		FORCEINLINE void CloseReadHandle();
-
-	protected:
-		virtual void SerializeElements(FArchive& Ar, void* Data) override;
-		virtual void SerializeElement(FArchive& Ar, void* Data, int64 ElementIndex) override { };
+		virtual int32 GetElementSize() const override { return ElementSize; }
+		virtual void SerializeElements(FArchive& Ar, void* Data) override { Octree->SerializeBulkData(Ar); }
 		virtual bool RequiresSingleElementSerialization(FArchive& Ar) override { return true; }
-	} BulkData;
-
+		void Serialize(FArchive& Ar);
+	} SavingBulkData;
+#endif
+	
 	TQueue<FLidarPointCloudOctreeNode*> QueuedNodes;
 	TArray<FLidarPointCloudOctreeNode*> NodesInUse;
 
@@ -555,6 +554,9 @@ public:
 
 	void OptimizeForStaticData();
 
+	IAsyncReadFileHandle* GetReadHandle();
+	void CloseReadHandle();
+
 	//~ Begin Deprecated
 	UE_DEPRECATED(4.27, "Use GetPointsInConvexVolume instead.")
 	void GetPointsInFrustum(TArray<FLidarPointCloudPoint*>& SelectedPoints, const FConvexVolume& Frustum, const bool& bVisibleOnly);
