Fixed Editor fails to build collision for LiDAR Point Cloud Asset

diff --git a/Engine/Plugins/Enterprise/LidarPointCloud/Source/LidarPointCloudRuntime/Private/Collision/LidarPointCloudCollision.cpp b/Engine/Plugins/Enterprise/LidarPointCloud/Source/LidarPointCloudRuntime/Private/Collision/LidarPointCloudCollision.cpp
--- a/Engine/Plugins/Enterprise/LidarPointCloud/Source/LidarPointCloudRuntime/Private/Collision/LidarPointCloudCollision.cpp
+++ b/Engine/Plugins/Enterprise/LidarPointCloud/Source/LidarPointCloudRuntime/Private/Collision/LidarPointCloudCollision.cpp
@@ -273,35 +273,35 @@ int8 TriangleTable[256][16] =
 /** Stores pre-computed data, ready for processing. */
 uint8 VertexData[256][181];
 
-FORCEINLINE FVector ExtractCoordinatesFromIndex(const int32& Index, const FIntVector& Dimensions)
+FORCEINLINE FVector3f ExtractCoordinatesFromIndex(const int32& Index, const FIntVector& Dimensions)
 {
 	const int32 W = Index % (Dimensions.X * Dimensions.Y);
-	return FVector(W % Dimensions.X, W / Dimensions.X, Index / (Dimensions.X * Dimensions.Y));
+	return FVector3f(W % Dimensions.X, W / Dimensions.X, Index / (Dimensions.X * Dimensions.Y));
 }
 
 void BuildScaledVertexTable(float CellSize)
 {
 	const float HalfCellSize = CellSize * 0.5f;
 
-	FVector VertexOffsetTable[12] =
+	FVector3f VertexOffsetTable[12] =
 	{
-		FVector(HalfCellSize, 0, 0),
-		FVector(CellSize, HalfCellSize, 0),
-		FVector(HalfCellSize, CellSize, 0),
-		FVector(0, HalfCellSize, 0),
-		FVector(HalfCellSize, 0, CellSize),
-		FVector(CellSize, HalfCellSize, CellSize),
-		FVector(HalfCellSize, CellSize, CellSize),
-		FVector(0, HalfCellSize, CellSize),
-		FVector(0, 0, HalfCellSize),
-		FVector(CellSize, 0, HalfCellSize),
-		FVector(CellSize, CellSize, HalfCellSize),
-		FVector(0, CellSize, HalfCellSize)
+		FVector3f(HalfCellSize, 0, 0),
+		FVector3f(CellSize, HalfCellSize, 0),
+		FVector3f(HalfCellSize, CellSize, 0),
+		FVector3f(0, HalfCellSize, 0),
+		FVector3f(HalfCellSize, 0, CellSize),
+		FVector3f(CellSize, HalfCellSize, CellSize),
+		FVector3f(HalfCellSize, CellSize, CellSize),
+		FVector3f(0, HalfCellSize, CellSize),
+		FVector3f(0, 0, HalfCellSize),
+		FVector3f(CellSize, 0, HalfCellSize),
+		FVector3f(CellSize, CellSize, HalfCellSize),
+		FVector3f(0, CellSize, HalfCellSize)
 	};
 
 	int8* VertexList = nullptr;
 
-	TArray<FVector> TempArray;
+	TArray<FVector3f> TempArray;
 	for (int32 i = 0; i < 256; ++i)
 	{
 		VertexList = TriangleTable[i];
@@ -396,9 +396,9 @@ void LidarPointCloudCollision::BuildCollisionMesh(FLidarPointCloudOctree* Octree
 	FBenchmarkTimer::Reset();
 
 	// Expand original bounds to make sure the mesh is closed at the edges
-	const FBox Bounds = Octree->GetBounds().ExpandBy(FVector::OneVector * CellSize, FVector::OneVector * CellSize);
-	const FVector BoundsSize = Bounds.GetSize();
-	const FVector LocationOffset = Octree->GetOwner()->LocationOffset.ToVector();
+	const FBox Bounds = Octree->GetBounds().ExpandBy(FVector3f::OneVector * CellSize, FVector3f::OneVector * CellSize);
+	const FVector3f BoundsSize = Bounds.GetSize();
+	const FVector3f LocationOffset = Octree->GetOwner()->LocationOffset.ToVector();
 
 	// Number of cells in each axis
 	const int32 BatchSize = GetDefault<ULidarPointCloudSettings>()->MeshingBatchSize;
@@ -469,7 +469,7 @@ void LidarPointCloudCollision::BuildCollisionMesh(FLidarPointCloudOctree* Octree
 						for (const FLidarPointCloudPoint** Point = Selection.GetData(), **DataEnd = Selection.GetData() + Selection.Num(); Point != DataEnd; ++Point)
 						{
 							// Calculate location relative to sampling bounds
-							FVector Location = (*Point)->Location - SamplingBounds.Min;
+							FVector3f Location = (*Point)->Location - SamplingBounds.Min;
 
 							// Calculate grid coordinates
 							const FIntVector Grid(FMath::Min(BatchSize - 1, (int32)(Location.X * InversedCellSize)), FMath::Min(BatchSize - 1, (int32)(Location.Y * InversedCellSize)), FMath::Min(BatchSize - 1, (int32)(Location.Z * InversedCellSize)));
@@ -503,7 +503,7 @@ void LidarPointCloudCollision::BuildCollisionMesh(FLidarPointCloudOctree* Octree
 	bool bWeldVertices = false;
 	if (bWeldVertices)
 	{
-		TSet<FVector> TmpVertexData;
+		TSet<FVector3f> TmpVertexData;
 		TmpVertexData.Reserve(NumVertices.GetValue());
 		CollisionMesh->Vertices.Empty(NumVertices.GetValue());
 		CollisionMesh->Indices.Empty(NumVertices.GetValue() / 3);
