UE-131882 - Fix LWC_TODO in LidarPointCloudShared.h

We now support FVector3f as a reflected type. We can remove the LWC_TODO.

diff --git a/Engine/Plugins/Enterprise/LidarPointCloud/Source/LidarPointCloudRuntime/Public/LidarPointCloudShared.h b/Engine/Plugins/Enterprise/LidarPointCloud/Source/LidarPointCloudRuntime/Public/LidarPointCloudShared.h
--- a/Engine/Plugins/Enterprise/LidarPointCloud/Source/LidarPointCloudRuntime/Public/LidarPointCloudShared.h
+++ b/Engine/Plugins/Enterprise/LidarPointCloud/Source/LidarPointCloudRuntime/Public/LidarPointCloudShared.h
@@ -345,8 +345,7 @@ struct LIDARPOINTCLOUDRUNTIME_API FLidarPointCloudPoint
 	GENERATED_BODY()
 
 public:
-	// LWC_TODO: Uncomment when FVector3f supports BP
-	//UPROPERTY(EditAnywhere, BlueprintReadWrite, Category = "Lidar Point Cloud Point")
+	UPROPERTY(EditAnywhere, BlueprintReadWrite, Category = "Lidar Point Cloud Point")
 	FVector3f Location;
 
 	UPROPERTY(EditAnywhere, BlueprintReadWrite, Category = "Lidar Point Cloud Point")
