LiDAR Point Cloud Asset opens with errors in Output Log for the first time

diff --git a/Engine/Plugins/Enterprise/LidarPointCloud/Source/LidarPointCloudEditor/Private/LidarPointCloudEditorViewport.cpp b/Engine/Plugins/Enterprise/LidarPointCloud/Source/LidarPointCloudEditor/Private/LidarPointCloudEditorViewport.cpp
--- a/Engine/Plugins/Enterprise/LidarPointCloud/Source/LidarPointCloudEditor/Private/LidarPointCloudEditorViewport.cpp
+++ b/Engine/Plugins/Enterprise/LidarPointCloud/Source/LidarPointCloudEditor/Private/LidarPointCloudEditorViewport.cpp
@@ -55,7 +55,7 @@ public:
 
 void SLidarPointCloudEditorViewport::Construct(const FArguments& InArgs)
 {
-	PreviewScene.SetFloorVisibility(false, true);
+	PreviewScene.SetFloorVisibility(false);
 	PreviewScene.SetEnvironmentVisibility(false);
 
 	PointCloudEditorPtr = InArgs._PointCloudEditor;
@@ -101,7 +101,7 @@ void SLidarPointCloudEditorViewport::Construct(const FArguments& InArgs)
 }
 
 SLidarPointCloudEditorViewport::SLidarPointCloudEditorViewport()
-	: PreviewScene(FPreviewScene::ConstructionValues(), FLT_MAX)
+	: PreviewScene(FPreviewScene::ConstructionValues())
 {
 }
 
