Fixed data corruption when importing E57 files
Added support for normals from E57

diff --git a/Engine/Plugins/Enterprise/LidarPointCloud/Source/LidarPointCloudRuntime/Private/IO/LidarPointCloudFileIO_E57.cpp b/Engine/Plugins/Enterprise/LidarPointCloud/Source/LidarPointCloudRuntime/Private/IO/LidarPointCloudFileIO_E57.cpp
--- a/Engine/Plugins/Enterprise/LidarPointCloud/Source/LidarPointCloudRuntime/Private/IO/LidarPointCloudFileIO_E57.cpp
+++ b/Engine/Plugins/Enterprise/LidarPointCloud/Source/LidarPointCloudRuntime/Private/IO/LidarPointCloudFileIO_E57.cpp
@@ -18,6 +18,7 @@ struct FE57Point
 {
 	FVector Location;
 	FColor Color;
+	FVector3f Normal;
 };
 #pragma pack(pop)
 
@@ -200,9 +201,10 @@ bool ULidarPointCloudFileIO_E57::HandleImport(const FString& Filename, TSharedPt
 	FCriticalSection SyncLock;
 	TArray<TFuture<void>> ThreadResults;
 
-	const uint32 MaxPointsToRead = FMath::Min(TotalPointsToRead, 2000000LL);
+	constexpr int64 MaxBufferSize = 2000000LL;
+	const uint32 MaxPointsToRead = FMath::Min(TotalPointsToRead, MaxBufferSize);
 
-	static FLidarPointCloudDataBufferManager BufferManager(MaxPointsToRead * sizeof(FE57Point));
+	static FLidarPointCloudDataBufferManager BufferManager(MaxBufferSize * sizeof(FE57Point));
 	FLidarPointCloudDataBuffer* Buffer = BufferManager.GetFreeBuffer();
 
 	// Read actual data
@@ -225,18 +227,20 @@ bool ULidarPointCloudFileIO_E57::HandleImport(const FString& Filename, TSharedPt
 
 					// Shift to protect from precision loss
 					Data->Location -= OutImportResults.OriginalCoordinates;
-
-					// Convert location to floats
-					const FVector ProcessedLocation = Data->Location;
-
-					Bounds += ProcessedLocation;
+					
+					Bounds += Data->Location;
 
 					if (!bHasRGB)
 					{
 						Data->Color.R = Data->Color.G = Data->Color.B = 255;
 					}
 
-					Points.Emplace((FVector3f)ProcessedLocation, Data->Color.B, Data->Color.G, Data->Color.R, bHasIntensity ? Data->Color.A : 255, 0);
+					if(!bHasIntensity)
+					{
+						Data->Color.A = 255;
+					}
+
+					Points.Emplace((FVector3f)Data->Location, Data->Color, true, 0, FLidarPointCloudNormal(Data->Normal));
 				}
 
 				ThreadBuffer->MarkAsFree();
