Fixed collision grid not showing after building collision mesh

diff --git a/Engine/Plugins/Enterprise/LidarPointCloud/Source/LidarPointCloudRuntime/Private/Rendering/LidarPointCloudRendering.cpp b/Engine/Plugins/Enterprise/LidarPointCloud/Source/LidarPointCloudRuntime/Private/Rendering/LidarPointCloudRendering.cpp
--- a/Engine/Plugins/Enterprise/LidarPointCloud/Source/LidarPointCloudRuntime/Private/Rendering/LidarPointCloudRendering.cpp
+++ b/Engine/Plugins/Enterprise/LidarPointCloud/Source/LidarPointCloudRuntime/Private/Rendering/LidarPointCloudRendering.cpp
@@ -178,7 +178,7 @@ public:
 		, ProxyWrapper(MakeShared<FLidarPointCloudSceneProxyWrapper, ESPMode::ThreadSafe>(this))
 		, bCompatiblePlatform(GetScene().GetFeatureLevel() >= ERHIFeatureLevel::SM5)
 		, Owner(Component->GetOwner())
-		, CollisionRendering(Component->GetPointCloud()->CollisionRendering)
+		, CollisionRenderingPtr(&Component->GetPointCloud()->CollisionRendering)
 	{
 		// Skip material verification - async update could occasionally cause it to crash
 		bVerifyUsedMaterials = false;
@@ -264,6 +264,7 @@ public:
 				}
 
 				// Render collision wireframe
+				FLidarPointCloudCollisionRendering* CollisionRendering = *CollisionRenderingPtr;
 				if (ViewFamily.EngineShowFlags.Collision && IsCollisionEnabled() && CollisionRendering && CollisionRendering->ShouldRenderCollision())
 				{
 					// Create colored proxy
@@ -406,7 +407,7 @@ private:
 	
 	AActor* Owner;
 
-	FLidarPointCloudCollisionRendering* CollisionRendering;
+	FLidarPointCloudCollisionRendering** CollisionRenderingPtr;
 };
 
 FPrimitiveSceneProxy* ULidarPointCloudComponent::CreateSceneProxy()
