Fixed ray-box intersection

diff --git a/Engine/Plugins/Enterprise/LidarPointCloud/Source/LidarPointCloudRuntime/Public/LidarPointCloudShared.h b/Engine/Plugins/Enterprise/LidarPointCloud/Source/LidarPointCloudRuntime/Public/LidarPointCloudShared.h
--- a/Engine/Plugins/Enterprise/LidarPointCloud/Source/LidarPointCloudRuntime/Public/LidarPointCloudShared.h
+++ b/Engine/Plugins/Enterprise/LidarPointCloud/Source/LidarPointCloudRuntime/Public/LidarPointCloudShared.h
@@ -339,9 +339,7 @@ public:
 	FORCEINLINE void SetDirection(const FVector3f& NewDirection)
 	{
 		Direction = NewDirection;
-		InvDirection = FVector3f(Direction.X == 0 ? 0 : 1 / Direction.X,
-								Direction.Y == 0 ? 0 : 1 / Direction.Y,
-								Direction.Z == 0 ? 0 : 1 / Direction.Z);
+		InvDirection = Direction.Reciprocal();
 	}
 
 	/** An Efficient and Robust Ray-Box Intersection Algorithm. Amy Williams et al. 2004. */
