Fixed incorrect point projection

diff --git a/Engine/Plugins/Enterprise/LidarPointCloud/Shaders/Private/LidarPointCloudVertexFactory.ush b/Engine/Plugins/Enterprise/LidarPointCloud/Shaders/Private/LidarPointCloudVertexFactory.ush
--- a/Engine/Plugins/Enterprise/LidarPointCloud/Shaders/Private/LidarPointCloudVertexFactory.ush
+++ b/Engine/Plugins/Enterprise/LidarPointCloud/Shaders/Private/LidarPointCloudVertexFactory.ush
@@ -324,7 +324,7 @@ FVertexFactoryIntermediates GetVertexFactoryIntermediates(FVertexFactoryInput In
 	Intermediates.TangentToLocal = CalcTangentToLocal(MetaData);
 	Intermediates.TangentToWorldSign = Intermediates.SceneData.InstanceData.DeterminantSign;
 	Intermediates.TangentToWorld = CalcTangentToWorldNoScale(Intermediates.TangentToLocal);
-	Intermediates.Position = ProcessPosition(Position, Intermediates.TangentToLocal, Input.VertexId);
+	Intermediates.Position = Position;
 	
 	// Color
 	{
@@ -381,15 +381,21 @@ half3x3 VertexFactoryGetTangentToLocal(FVertexFactoryInput Input, FVertexFactory
 // @return translated world position
 float4 VertexFactoryGetWorldPosition(FVertexFactoryInput Input, FVertexFactoryIntermediates Intermediates)
 {
-	return TransformLocalToTranslatedWorld(Intermediates.Position);
+	float4 WorldPosition = TransformLocalToTranslatedWorld(Intermediates.Position);
+	WorldPosition.xyz = ProcessPosition(WorldPosition.xyz, Intermediates.TangentToLocal, Input.VertexId);
+	return WorldPosition;
 }
 float4 VertexFactoryGetWorldPosition(FPositionOnlyVertexFactoryInput Input)
 {
-	return TransformLocalToTranslatedWorld(ProcessPosition(bUseStaticBuffers ? Input.Position.xyz : LidarGetInstancePosition(Input.VertexId), Input.VertexId));
+	float4 WorldPosition = TransformLocalToTranslatedWorld(bUseStaticBuffers ? Input.Position.xyz : LidarGetInstancePosition(Input.VertexId));
+	WorldPosition.xyz = ProcessPosition(WorldPosition.xyz, Input.VertexId);
+	return WorldPosition;
 }
 float4 VertexFactoryGetWorldPosition(FPositionAndNormalOnlyVertexFactoryInput Input)
 {
-	return TransformLocalToTranslatedWorld(ProcessPosition(bUseStaticBuffers ? Input.Position.xyz : LidarGetInstancePosition(Input.VertexId), Input.VertexId));
+	float4 WorldPosition = TransformLocalToTranslatedWorld(bUseStaticBuffers ? Input.Position.xyz : LidarGetInstancePosition(Input.VertexId));
+	WorldPosition.xyz = ProcessPosition(WorldPosition.xyz, Input.VertexId);
+	return WorldPosition;
 }
 float4 VertexFactoryGetRasterizedWorldPosition(FVertexFactoryInput Input, FVertexFactoryIntermediates Intermediates, float4 InWorldPosition)
 {
