Fixed incorrect classification id allocation
Added classification color alpha affecting final alpha

diff --git a/Engine/Plugins/Enterprise/LidarPointCloud/Shaders/Private/LidarPointCloudVertexFactory.ush b/Engine/Plugins/Enterprise/LidarPointCloud/Shaders/Private/LidarPointCloudVertexFactory.ush
--- a/Engine/Plugins/Enterprise/LidarPointCloud/Shaders/Private/LidarPointCloudVertexFactory.ush
+++ b/Engine/Plugins/Enterprise/LidarPointCloud/Shaders/Private/LidarPointCloudVertexFactory.ush
@@ -101,6 +101,7 @@ struct FVertexFactoryInterpolantsVSToPS
 	TANGENTTOWORLD_INTERPOLATOR_BLOCK
 
 	half4	Color : COLOR0;
+	half	Alpha : COLOR1;
 	float2	TexCoords : TEXCOORD0;
 
 #if INSTANCED_STEREO
@@ -114,6 +115,7 @@ struct FVertexFactoryIntermediates
     half3x3 TangentToWorld;
     half TangentToWorldSign;
     half4 Color;
+	half Alpha;
 	float3 Position;
 	/** Cached primitive and instance data */
 	FSceneDataIntermediates SceneData;
@@ -334,10 +336,10 @@ FVertexFactoryIntermediates GetVertexFactoryIntermediates(FVertexFactoryInput In
 
 		half3 ElevationColor = lerp(ElevationColorBottom, ElevationColorTop, clamp((Position.z + BoundsSize.z * 0.5) / BoundsSize.z, 0, 1));
 		half3 PositionColor = (Position + BoundsSize * 0.5) / BoundsSize;
-		half3 ClassificationColor = ClassificationColors[0x0000001F & (Classification >> 1)].rgb;
+		float4 ClassificationColor = ClassificationColors[Classification];
 
 		// Pick color source
-		Color = lerp(lerp(Color, lerp(lerp(PositionColor, ClassificationColor, bUseClassification), ElevationColor, bUseElevationColor), bUseColorOverride), SelectionColor, bSelected);
+		Color = lerp(lerp(Color, lerp(lerp(PositionColor, ClassificationColor.rgb, bUseClassification), ElevationColor, bUseElevationColor), bUseColorOverride), SelectionColor, bSelected);
 
 		// Mix with Intensity
 		Color = lerp(Color, Color * Intensity, IntensityInfluence);
@@ -353,6 +355,7 @@ FVertexFactoryIntermediates GetVertexFactoryIntermediates(FVertexFactoryInput In
 		}
 
 		Intermediates.Color = half4(Color, bUseCircle);
+		Intermediates.Alpha = lerp(1, ClassificationColor.a, bUseClassification);
 	}
 
 	return Intermediates;
@@ -364,6 +367,7 @@ FVertexFactoryInterpolantsVSToPS VertexFactoryGetInterpolantsVSToPS(FVertexFacto
 	Interpolants.TangentToWorld0 = float4(Intermediates.TangentToWorld[0], 0);
 	Interpolants.TangentToWorld2 = float4(Intermediates.TangentToWorld[2], Intermediates.TangentToWorldSign);
 	Interpolants.Color = VertexParameters.VertexColor;
+	Interpolants.Alpha = Intermediates.Alpha;
 	Interpolants.TexCoords = PrecomputedUV_Quad[Input.VertexId % 4 + 4];    // Using inverted V
 
 #if INSTANCED_STEREO
@@ -445,7 +449,7 @@ FMaterialPixelParameters GetMaterialPixelParameters(FVertexFactoryInterpolantsVS
     Result.Particle.Color = half4(1, 1, 1, 1);
     Result.TwoSidedSign = 1;
 	Result.VertexColor = Interpolants.Color;
-	Result.VertexColor.a = lerp(1, 1 - distance(Interpolants.TexCoords, 0.5), Interpolants.Color.a);
+	Result.VertexColor.a = saturate((lerp(1, 1 - distance(Interpolants.TexCoords, 0.5), Interpolants.Color.a) - 0.5) * 10000) * Interpolants.Alpha;
 
     half3 TangentToWorld0 = Interpolants.TangentToWorld0.xyz;
     half4 TangentToWorld2 = Interpolants.TangentToWorld2;
