Added DataWithClassificationAlpha color source

diff --git a/Engine/Plugins/Enterprise/LidarPointCloud/Shaders/Private/LidarPointCloudVertexFactory.ush b/Engine/Plugins/Enterprise/LidarPointCloud/Shaders/Private/LidarPointCloudVertexFactory.ush
--- a/Engine/Plugins/Enterprise/LidarPointCloud/Shaders/Private/LidarPointCloudVertexFactory.ush
+++ b/Engine/Plugins/Enterprise/LidarPointCloud/Shaders/Private/LidarPointCloudVertexFactory.ush
@@ -53,6 +53,7 @@ uint bUseCircle;
 uint bUseColorOverride;
 uint bUseElevationColor;
 uint bUseClassification;
+uint bUseClassificationAlpha;
 
 float4 ClassificationColors[32];
 
@@ -355,7 +356,7 @@ FVertexFactoryIntermediates GetVertexFactoryIntermediates(FVertexFactoryInput In
 		}
 
 		Intermediates.Color = half4(Color, bUseCircle);
-		Intermediates.Alpha = lerp(1, ClassificationColor.a, bUseClassification);
+		Intermediates.Alpha = lerp(1, ClassificationColor.a, bUseClassification || bUseClassificationAlpha);
 	}
 
 	return Intermediates;
diff --git a/Engine/Plugins/Enterprise/LidarPointCloud/Source/LidarPointCloudRuntime/Private/Rendering/LidarPointCloudRenderBuffers.cpp b/Engine/Plugins/Enterprise/LidarPointCloud/Source/LidarPointCloudRuntime/Private/Rendering/LidarPointCloudRenderBuffers.cpp
--- a/Engine/Plugins/Enterprise/LidarPointCloud/Source/LidarPointCloudRuntime/Private/Rendering/LidarPointCloudRenderBuffers.cpp
+++ b/Engine/Plugins/Enterprise/LidarPointCloud/Source/LidarPointCloudRuntime/Private/Rendering/LidarPointCloudRenderBuffers.cpp
@@ -169,6 +169,7 @@ void FLidarPointCloudVertexFactoryShaderParameters::Bind(const FShaderParameterM
 	BINDPARAM(Tint);
 	BINDPARAM(IntensityInfluence);
 	BINDPARAM(bUseClassification);
+	BINDPARAM(bUseClassificationAlpha);
 	BINDPARAM(ClassificationColors);
 	BINDPARAM(ClippingVolume);
 	BINDPARAM(NumClippingVolumes);
@@ -215,6 +216,7 @@ void FLidarPointCloudVertexFactoryShaderParameters::GetElementShaderBindings(con
 	SETPARAM(Tint);
 	SETPARAM(IntensityInfluence);
 	SETPARAM(bUseClassification);
+	SETPARAM(bUseClassificationAlpha);
 	SETPARAM(ClassificationColors);
 	SETPARAM(ClippingVolume);
 	SETPARAM(NumClippingVolumes);
diff --git a/Engine/Plugins/Enterprise/LidarPointCloud/Source/LidarPointCloudRuntime/Private/Rendering/LidarPointCloudRenderBuffers.h b/Engine/Plugins/Enterprise/LidarPointCloud/Source/LidarPointCloudRuntime/Private/Rendering/LidarPointCloudRenderBuffers.h
--- a/Engine/Plugins/Enterprise/LidarPointCloud/Source/LidarPointCloudRuntime/Private/Rendering/LidarPointCloudRenderBuffers.h
+++ b/Engine/Plugins/Enterprise/LidarPointCloud/Source/LidarPointCloudRuntime/Private/Rendering/LidarPointCloudRenderBuffers.h
@@ -94,6 +94,7 @@ struct FLidarPointCloudBatchElementUserData
 	FVector3f Tint;
 	float IntensityInfluence;
 	int32 bUseClassification;
+	int32 bUseClassificationAlpha;
 	FVector4f ClassificationColors[32];
 	FMatrix44f ClippingVolume[16];
 	uint32 NumClippingVolumes;
@@ -145,6 +146,7 @@ public:
 	LAYOUT_FIELD(FShaderParameter, Tint);
 	LAYOUT_FIELD(FShaderParameter, IntensityInfluence);
 	LAYOUT_FIELD(FShaderParameter, bUseClassification);
+	LAYOUT_FIELD(FShaderParameter, bUseClassificationAlpha);
 	LAYOUT_FIELD(FShaderParameter, ClassificationColors);
 	LAYOUT_FIELD(FShaderParameter, ClippingVolume);
 	LAYOUT_FIELD(FShaderParameter, NumClippingVolumes);
diff --git a/Engine/Plugins/Enterprise/LidarPointCloud/Source/LidarPointCloudRuntime/Private/Rendering/LidarPointCloudRendering.cpp b/Engine/Plugins/Enterprise/LidarPointCloud/Source/LidarPointCloudRuntime/Private/Rendering/LidarPointCloudRendering.cpp
--- a/Engine/Plugins/Enterprise/LidarPointCloud/Source/LidarPointCloudRuntime/Private/Rendering/LidarPointCloudRendering.cpp
+++ b/Engine/Plugins/Enterprise/LidarPointCloud/Source/LidarPointCloudRuntime/Private/Rendering/LidarPointCloudRendering.cpp
@@ -340,7 +340,7 @@ public:
 		UserDataElement.ElevationColorBottom = FVector3f(RenderData.RenderParams.ColorSource == ELidarPointCloudColorationMode::None ? FColor::White : RenderData.RenderParams.ElevationColorBottom);
 		UserDataElement.ElevationColorTop = FVector3f(RenderData.RenderParams.ColorSource == ELidarPointCloudColorationMode::None ? FColor::White : RenderData.RenderParams.ElevationColorTop);
 		UserDataElement.bUseCircle = bUsesSprites && RenderData.RenderParams.PointShape == ELidarPointCloudSpriteShape::Circle;
-		UserDataElement.bUseColorOverride = RenderData.RenderParams.ColorSource != ELidarPointCloudColorationMode::Data;
+		UserDataElement.bUseColorOverride = RenderData.RenderParams.ColorSource != ELidarPointCloudColorationMode::Data && RenderData.RenderParams.ColorSource != ELidarPointCloudColorationMode::DataWithClassificationAlpha;
 		UserDataElement.bUseElevationColor = RenderData.RenderParams.ColorSource == ELidarPointCloudColorationMode::Elevation || RenderData.RenderParams.ColorSource == ELidarPointCloudColorationMode::None;
 		UserDataElement.Offset = RenderData.RenderParams.Offset;
 		UserDataElement.Contrast = RenderData.RenderParams.Contrast;
@@ -350,6 +350,7 @@ public:
 		UserDataElement.IntensityInfluence = RenderData.RenderParams.IntensityInfluence;
 
 		UserDataElement.bUseClassification = RenderData.RenderParams.ColorSource == ELidarPointCloudColorationMode::Classification;
+		UserDataElement.bUseClassificationAlpha = RenderData.RenderParams.ColorSource == ELidarPointCloudColorationMode::DataWithClassificationAlpha;
 		UserDataElement.SetClassificationColors(RenderData.RenderParams.ClassificationColors);
 				
 		UserDataElement.NumClippingVolumes = FMath::Min(RenderData.ClippingVolumes.Num(), 16);
diff --git a/Engine/Plugins/Enterprise/LidarPointCloud/Source/LidarPointCloudRuntime/Public/LidarPointCloudShared.h b/Engine/Plugins/Enterprise/LidarPointCloud/Source/LidarPointCloudRuntime/Public/LidarPointCloudShared.h
--- a/Engine/Plugins/Enterprise/LidarPointCloud/Source/LidarPointCloudRuntime/Public/LidarPointCloudShared.h
+++ b/Engine/Plugins/Enterprise/LidarPointCloud/Source/LidarPointCloudRuntime/Public/LidarPointCloudShared.h
@@ -425,6 +425,8 @@ enum class ELidarPointCloudColorationMode : uint8
 	None,
 	/** Uses imported RGB / Intensity data */
 	Data,
+	/** Uses imported RGB / Intensity data combined with Alpha mask from Classification Colors */
+	DataWithClassificationAlpha,
 	/** The cloud's color will be overridden with elevation-based color */
 	Elevation,
 	/** The cloud's color will be overridden with relative position-based color */
