Fixed crash when building collision

diff --git a/Engine/Plugins/Enterprise/LidarPointCloud/Source/LidarPointCloudRuntime/Private/LidarPointCloud.cpp b/Engine/Plugins/Enterprise/LidarPointCloud/Source/LidarPointCloudRuntime/Private/LidarPointCloud.cpp
--- a/Engine/Plugins/Enterprise/LidarPointCloud/Source/LidarPointCloudRuntime/Private/LidarPointCloud.cpp
+++ b/Engine/Plugins/Enterprise/LidarPointCloud/Source/LidarPointCloudRuntime/Private/LidarPointCloud.cpp
@@ -9,26 +9,24 @@
 #include "Async/Future.h"
 #include "Serialization/CustomVersion.h"
 #include "Misc/ScopeTryLock.h"
-#include "Misc/ConfigCacheIni.h"
 #include "Misc/ScopedSlowTask.h"
 #include "Engine/Engine.h"
 #include "LatentActions.h"
 #include "PhysicsEngine/BodySetup.h"
 #include "Framework/Notifications/NotificationManager.h"
-#include "Engine/CollisionProfile.h"
 #include "EngineUtils.h"
 #include "Components/BrushComponent.h"
+#include "UObject/GCObjectScopeGuard.h"
 #include "Widgets/Notifications/SNotificationList.h"
 #include "UObject/ObjectSaveContext.h"
+#include "UObject/UObjectGlobals.h"
 
 #if WITH_EDITOR
 #include "IContentBrowserSingleton.h"
 #include "ContentBrowserModule.h"
 #include "AssetRegistry/AssetRegistryModule.h"
-#include "Editor.h"
 #include "Styling/SlateStyleRegistry.h"
 #include "Misc/MessageDialog.h"
-#include "UObject/UObjectGlobals.h"
 #endif
 
 #define IS_PROPERTY(Name) PropertyChangedEvent.MemberProperty->GetName().Equals(#Name)
@@ -605,18 +603,26 @@ void ULidarPointCloud::BuildCollision()
 	NewBodySetup->CollisionTraceFlag = CTF_UseComplexAsSimple;
 	NewBodySetup->bHasCookedCollisionData = true;
 
-	Async(EAsyncExecution::Thread, [this, Notification]{
-		Octree.BuildCollision(MaxCollisionError, true);
+	TWeakObjectPtr<ULidarPointCloud> WeakThis = this;
 
-		FBenchmarkTimer::Reset();
-#if WITH_PHYSX  && PHYSICS_INTERFACE_PHYSX
-		AsyncTask(ENamedThreads::GameThread, [this, Notification] {
-			NewBodySetup->CreatePhysicsMeshesAsync(FOnAsyncPhysicsCookFinished::CreateUObject(this, &ULidarPointCloud::FinishPhysicsAsyncCook, Notification));
-		});
-#elif WITH_CHAOS
-		NewBodySetup->CreatePhysicsMeshes();
-		AsyncTask(ENamedThreads::GameThread, [this, Notification] { FinishPhysicsAsyncCook(true, Notification); });
-#endif		
+	Async(EAsyncExecution::Thread, [WeakThis = MoveTemp(WeakThis), Notification]() mutable
+	{
+		if (ULidarPointCloud* RawPtr = WeakThis.Get())
+		{
+			FGCObjectScopeGuard Guard(RawPtr);
+			RawPtr->Octree.BuildCollision(RawPtr->MaxCollisionError, true);
+
+			FBenchmarkTimer::Reset();
+
+			RawPtr->NewBodySetup->CreatePhysicsMeshes();
+			AsyncTask(ENamedThreads::GameThread, [WeakThis = MoveTemp(WeakThis), Notification]
+			{
+				if (ULidarPointCloud* RawPtr = WeakThis.Get())
+				{
+					RawPtr->FinishPhysicsAsyncCook(true, Notification);
+				}
+			});
+		}
 	});
 }
 
