Fixed crash when streaming nodes

diff --git a/Engine/Plugins/Enterprise/LidarPointCloud/Source/LidarPointCloudRuntime/Private/LidarPointCloudLODManager.cpp b/Engine/Plugins/Enterprise/LidarPointCloud/Source/LidarPointCloudRuntime/Private/LidarPointCloudLODManager.cpp
--- a/Engine/Plugins/Enterprise/LidarPointCloud/Source/LidarPointCloudRuntime/Private/LidarPointCloudLODManager.cpp
+++ b/Engine/Plugins/Enterprise/LidarPointCloud/Source/LidarPointCloudRuntime/Private/LidarPointCloudLODManager.cpp
@@ -661,8 +661,12 @@ int64 FLidarPointCloudLODManager::ProcessLOD(const TArray<FLidarPointCloudLODMan
 		int32 LoadedNodes = 0;
 		for (TPair<FLidarPointCloudOctree*, TArray<FLidarPointCloudOctreeNode*>>& OctreeStreamingData : OctreeStreamingMap)
 		{
-			OctreeStreamingData.Key->StreamNodes(OctreeStreamingData.Value, CurrentTime);
-			LoadedNodes += OctreeStreamingData.Key->GetNumNodesInUse();
+			FScopeTryLock OctreeLock(&OctreeStreamingData.Key->DataLock);
+			if (OctreeLock.IsLocked())
+			{
+				OctreeStreamingData.Key->StreamNodes(OctreeStreamingData.Value, CurrentTime);
+				LoadedNodes += OctreeStreamingData.Key->GetNumNodesInUse();
+			}
 		}
 
 		SET_DWORD_STAT(STAT_LoadedNodes, LoadedNodes);
diff --git a/Engine/Plugins/Enterprise/LidarPointCloud/Source/LidarPointCloudRuntime/Private/LidarPointCloudOctree.cpp b/Engine/Plugins/Enterprise/LidarPointCloud/Source/LidarPointCloudRuntime/Private/LidarPointCloudOctree.cpp
--- a/Engine/Plugins/Enterprise/LidarPointCloud/Source/LidarPointCloudRuntime/Private/LidarPointCloudOctree.cpp
+++ b/Engine/Plugins/Enterprise/LidarPointCloud/Source/LidarPointCloudRuntime/Private/LidarPointCloudOctree.cpp
@@ -865,6 +865,7 @@ FLidarPointCloudOctree::FLidarPointCloudOctree(ULidarPointCloud* Owner)
 
 FLidarPointCloudOctree::~FLidarPointCloudOctree()
 {
+	FScopeLock Lock(&DataLock);
 	Empty(true);
 }
 
