FBulkData changes for Lidar

diff --git a/Engine/Plugins/Enterprise/LidarPointCloud/Source/LidarPointCloudRuntime/Private/LidarPointCloudOctree.cpp b/Engine/Plugins/Enterprise/LidarPointCloud/Source/LidarPointCloudRuntime/Private/LidarPointCloudOctree.cpp
--- a/Engine/Plugins/Enterprise/LidarPointCloud/Source/LidarPointCloudRuntime/Private/LidarPointCloudOctree.cpp
+++ b/Engine/Plugins/Enterprise/LidarPointCloud/Source/LidarPointCloudRuntime/Private/LidarPointCloudOctree.cpp
@@ -865,7 +865,6 @@ FLidarPointCloudOctree::FLidarPointCloudOctree(ULidarPointCloud* Owner)
 
 FLidarPointCloudOctree::~FLidarPointCloudOctree()
 {
-	FScopeLock Lock(&DataLock);
 	Empty(true);
 }
 
@@ -2478,7 +2477,7 @@ void FLidarPointCloudOctree::Serialize(FArchive& Ar)
 #if WITH_EDITOR
 		if(Ar.IsSaving())
 		{
-			SavingBulkData.Serialize(Ar);
+			SavingBulkData.Serialize(Ar, Owner, false, 1, EFileRegionType::None);
 		}
 		else
 #endif
@@ -2570,19 +2569,17 @@ void FLidarPointCloudOctree::StreamNodeData(FLidarPointCloudOctreeNode* Node)
 //////////////////////////////////////////////////////////// FLidarPointCloudBulkData
 
 #if WITH_EDITOR
-FLidarPointCloudOctree::FLidarPointCloudBulkData::FLidarPointCloudBulkData(FLidarPointCloudOctree* Octree): Octree(Octree)
-{
-	SetBulkDataFlags(BULKDATA_Force_NOT_InlinePayload);
-	Lock(LOCK_READ_WRITE);
-	Realloc(2);
-	Unlock();
-}
-
-void FLidarPointCloudOctree::FLidarPointCloudBulkData::Serialize(FArchive& Ar)
+FLidarPointCloudOctree::FLidarPointCloudBulkData::FLidarPointCloudBulkData(FLidarPointCloudOctree* Octree)
 {
-	ElementSize = Ar.IsPersistent() && !Ar.IsObjectReferenceCollector() && !Ar.ShouldSkipBulkData() ? INT32_MAX : 1;
-	FUntypedBulkData::Serialize(Ar, Octree->GetOwner());
-	ElementSize = 1;
+	SerializeElementsCallback = [Octree](FArchive& Ar, void*, int64, EBulkDataFlags)
+	{
+		Octree->SerializeBulkData(Ar);
+	};
+#if !USE_RUNTIME_BULKDATA
+	SerializeBulkDataElements = &SerializeElementsCallback;
+#endif // !USE_RUNTIME_BULKDATA
+	
+	SetBulkDataFlags(BULKDATA_Force_NOT_InlinePayload | BULKDATA_Size64Bit);
 }
 #endif
 
diff --git a/Engine/Plugins/Enterprise/LidarPointCloud/Source/LidarPointCloudRuntime/Public/LidarPointCloudOctree.h b/Engine/Plugins/Enterprise/LidarPointCloud/Source/LidarPointCloudRuntime/Public/LidarPointCloudOctree.h
--- a/Engine/Plugins/Enterprise/LidarPointCloud/Source/LidarPointCloudRuntime/Public/LidarPointCloudOctree.h
+++ b/Engine/Plugins/Enterprise/LidarPointCloud/Source/LidarPointCloudRuntime/Public/LidarPointCloudOctree.h
@@ -272,21 +272,11 @@ private:
 	FByteBulkData BulkData;
 
 #if WITH_EDITOR
-	struct FLidarPointCloudBulkData : public FUntypedBulkData
+	struct FLidarPointCloudBulkData : public FBulkData
 	{
-		FLidarPointCloudOctree* Octree;
-		int32 ElementSize = 1;
+		FSerializeBulkDataElements SerializeElementsCallback;
 		
 		FLidarPointCloudBulkData(FLidarPointCloudOctree* Octree);
-		virtual int32 GetElementSize() const override { return ElementSize; }
-		virtual void SerializeElements(FArchive& Ar, void* Data) override { Octree->SerializeBulkData(Ar); }
-		virtual void SerializeElement(FArchive& Ar, void* Data, int64 ElementIndex) override
-		{
-			uint8* Element = static_cast<uint8*>(Data) + (ElementIndex * ElementSize); 
-			Ar.Serialize(Element, ElementSize);
-		}
-		virtual bool RequiresSingleElementSerialization(FArchive& Ar) override { return true; }
-		void Serialize(FArchive& Ar);
 	} SavingBulkData;
 #endif
 	
