New editor mode for Lidar

diff --git a/Engine/Plugins/Enterprise/LidarPointCloud/Shaders/Private/LidarPointCloudVertexFactory.ush b/Engine/Plugins/Enterprise/LidarPointCloud/Shaders/Private/LidarPointCloudVertexFactory.ush
--- a/Engine/Plugins/Enterprise/LidarPointCloud/Shaders/Private/LidarPointCloudVertexFactory.ush
+++ b/Engine/Plugins/Enterprise/LidarPointCloud/Shaders/Private/LidarPointCloudVertexFactory.ush
@@ -332,7 +332,7 @@ FVertexFactoryIntermediates GetVertexFactoryIntermediates(FVertexFactoryInput In
 	// Color
 	{
 		MetaData = 0x000000FF & (MetaData >> 24);
-		uint bSelected = bEditorView && (0x00000040 & MetaData);
+		uint bSelected = clamp(0x00000040 & MetaData, 0, 1);
 		uint Classification = 0x0000001F & (MetaData >> 1);
 
 		half3 ElevationColor = lerp(ElevationColorBottom, ElevationColorTop, clamp((Position.z + BoundsSize.z * 0.5) / BoundsSize.z, 0, 1));
@@ -340,7 +340,7 @@ FVertexFactoryIntermediates GetVertexFactoryIntermediates(FVertexFactoryInput In
 		float4 ClassificationColor = ClassificationColors[Classification];
 
 		// Pick color source
-		Color = lerp(lerp(Color, lerp(lerp(PositionColor, ClassificationColor.rgb, bUseClassification), ElevationColor, bUseElevationColor), bUseColorOverride), SelectionColor, bSelected);
+		Color = lerp(Color, lerp(lerp(PositionColor, ClassificationColor.rgb, bUseClassification), ElevationColor, bUseElevationColor), bUseColorOverride);
 
 		// Mix with Intensity
 		Color = lerp(Color, Color * Intensity, IntensityInfluence);
@@ -355,6 +355,9 @@ FVertexFactoryIntermediates GetVertexFactoryIntermediates(FVertexFactoryInput In
 			Color = clamp(Color, 0, 1);
 		}
 
+		// Selection
+		Color = lerp(Color, SelectionColor, bSelected);
+
 		Intermediates.Color = half4(Color, bUseCircle);
 		Intermediates.Alpha = lerp(1, ClassificationColor.a, bUseClassification || bUseClassificationAlpha);
 	}
diff --git a/Engine/Plugins/Enterprise/LidarPointCloud/Source/LidarPointCloudEditor/LidarPointCloudEditor.Build.cs b/Engine/Plugins/Enterprise/LidarPointCloud/Source/LidarPointCloudEditor/LidarPointCloudEditor.Build.cs
--- a/Engine/Plugins/Enterprise/LidarPointCloud/Source/LidarPointCloudEditor/LidarPointCloudEditor.Build.cs
+++ b/Engine/Plugins/Enterprise/LidarPointCloud/Source/LidarPointCloudEditor/LidarPointCloudEditor.Build.cs
@@ -30,7 +30,10 @@ namespace UnrealBuildTool.Rules
                 "Engine",
                 "RenderCore",
                 "RHI",
-                "AdvancedPreviewScene"
+                "AdvancedPreviewScene",
+                "EditorInteractiveToolsFramework",
+                "InteractiveToolsFramework",
+                "StatusBar"
                 }
             );
         }
diff --git a/Engine/Plugins/Enterprise/LidarPointCloud/Source/LidarPointCloudEditor/Private/LidarPointCloudEdMode.cpp b/Engine/Plugins/Enterprise/LidarPointCloud/Source/LidarPointCloudEditor/Private/LidarPointCloudEdMode.cpp
new file mode 100644
--- /dev/null
+++ b/Engine/Plugins/Enterprise/LidarPointCloud/Source/LidarPointCloudEditor/Private/LidarPointCloudEdMode.cpp
@@ -0,0 +1,139 @@
+﻿// Copyright Epic Games, Inc. All Rights Reserved.
+
+#include "LidarPointCloudEdMode.h"
+
+#include "ContentBrowserModule.h"
+#include "EditorModeManager.h"
+#include "InteractiveToolManager.h"
+#include "LevelEditorViewport.h"
+#include "LidarPointCloudActor.h"
+#include "LidarPointCloudEditorCommands.h"
+#include "LidarPointCloudEditorTools.h"
+#include "LidarPointCloudEdModeToolkit.h"
+#include "Selection.h"
+
+#define LOCTEXT_NAMESPACE "LidarPointCloudEdMode"
+
+namespace FLidarEditorModes
+{
+	const FEditorModeID EM_Lidar(TEXT("EM_Lidar"));
+}
+
+ULidarEditorMode::ULidarEditorMode()
+	: Super()
+{
+	Info = FEditorModeInfo(
+		FLidarEditorModes::EM_Lidar,
+		NSLOCTEXT("EditorModes", "LidarMode", "Lidar"),
+		FSlateIcon(FLidarPointCloudStyle::GetStyleSetName(), "ClassThumbnail.LidarPointCloud", "ClassIcon.LidarPointCloud"),
+		true
+		);
+}
+
+void ULidarEditorMode::Enter()
+{
+	Super::Enter();
+	
+	GEditor->SelectNone(true, true);
+
+	const FLidarPointCloudEditorCommands& Commands = FLidarPointCloudEditorCommands::Get();
+
+#define REGISTER_TOOL(Tool) RegisterTool(Commands.Toolkit##Tool, TEXT("Lidar"#Tool"Tool"), NewObject<ULidarEditorToolBuilder_##Tool>())
+	REGISTER_TOOL(Select);
+	REGISTER_TOOL(Align);
+	REGISTER_TOOL(Merge);
+	REGISTER_TOOL(Collision);
+	REGISTER_TOOL(Normals);
+	REGISTER_TOOL(BoxSelection);
+	REGISTER_TOOL(PolygonalSelection);
+	REGISTER_TOOL(LassoSelection);
+	REGISTER_TOOL(PaintSelection);
+#undef REGISTER_TOOL
+	
+	PaletteChangedHandle = Toolkit->OnPaletteChanged().AddUObject(this, &ULidarEditorMode::UpdateOnPaletteChange);
+	
+	GetToolManager()->ConfigureChangeTrackingMode(EToolChangeTrackingMode::NoChangeTracking);
+	Toolkit->SetCurrentPalette(LidarEditorPalletes::Manage);
+	GetToolManager()->ConfigureChangeTrackingMode(EToolChangeTrackingMode::UndoToExit);
+}
+
+void ULidarEditorMode::Exit()
+{
+	Toolkit->OnPaletteChanged().Remove(PaletteChangedHandle);
+	
+	Super::Exit();
+}
+
+bool ULidarEditorMode::IsSelectionAllowed(AActor* InActor, bool bInSelection) const
+{
+	return InActor && InActor->IsA(ALidarPointCloudActor::StaticClass());
+}
+
+void ULidarEditorMode::CreateToolkit()
+{
+	Toolkit = MakeShareable(new FLidarPointCloudEdModeToolkit);
+}
+
+bool ULidarEditorMode::GetPivotForOrbit(FVector& OutPivot) const
+{
+	if (GCurrentLevelEditingViewportClient)
+	{
+		OutPivot = GCurrentLevelEditingViewportClient->GetViewTransform().GetLookAt();
+		return true;
+	}
+	return false;
+}
+
+TMap<FName, TArray<TSharedPtr<FUICommandInfo>>> ULidarEditorMode::GetModeCommands() const
+{
+	return FLidarPointCloudEditorCommands::GetCommands();
+}
+
+void ULidarEditorMode::UpdateOnPaletteChange(FName NewPalette)
+{
+	GetToolManager()->SelectActiveToolType(EToolSide::Mouse, "LidarSelectTool");
+	GetToolManager()->ActivateTool(EToolSide::Mouse);
+}
+
+void ULidarEditorMode::CancelActiveToolAction()
+{
+	if (GetToolManager()->HasAnyActiveTool())
+	{
+		if (IInteractiveToolNestedAcceptCancelAPI* CancelAPI = Cast<IInteractiveToolNestedAcceptCancelAPI>(GetToolManager()->GetActiveTool(EToolSide::Mouse)))
+		{
+			CancelAPI->ExecuteNestedCancelCommand();
+		}
+	}
+}
+
+void ULidarEditorMode::BindCommands()
+{
+	const FLidarPointCloudEditorCommands& Commands = FLidarPointCloudEditorCommands::Get();
+	const TSharedRef<FUICommandList>& CommandList = Toolkit->GetToolkitCommands();
+
+	CommandList->MapAction(
+		Commands.ToolktitCancelSelection,
+		FExecuteAction::CreateUObject(this, &ULidarEditorMode::CancelActiveToolAction),
+		FCanExecuteAction(),
+		FGetActionCheckState(),
+		FIsActionButtonVisible(),
+		EUIActionRepeatMode::RepeatDisabled);
+}
+
+void ULidarEditorMode::OnToolStarted(UInteractiveToolManager* Manager, UInteractiveTool* Tool)
+{
+	if(ULidarEditorTool_Base* LidarTool = Cast<ULidarEditorTool_Base>(Tool))
+	{
+		GetToolManager()->DisplayMessage(LidarTool->GetToolMessage(), EToolMessageLevel::UserNotification);
+	}
+}
+
+void ULidarEditorMode::OnToolEnded(UInteractiveToolManager* Manager, UInteractiveTool* Tool)
+{
+	if(Tool && Tool->IsA(ULidarEditorTool_SelectionBase::StaticClass()))
+	{
+		FLidarPointCloudEditorHelper::ClearSelection();
+	}
+}
+
+#undef LOCTEXT_NAMESPACE
diff --git a/Engine/Plugins/Enterprise/LidarPointCloud/Source/LidarPointCloudEditor/Private/LidarPointCloudEdMode.h b/Engine/Plugins/Enterprise/LidarPointCloud/Source/LidarPointCloudEditor/Private/LidarPointCloudEdMode.h
new file mode 100644
--- /dev/null
+++ b/Engine/Plugins/Enterprise/LidarPointCloud/Source/LidarPointCloudEditor/Private/LidarPointCloudEdMode.h
@@ -0,0 +1,39 @@
+﻿// Copyright Epic Games, Inc. All Rights Reserved.
+
+#pragma once
+#include "Tools/LegacyEdModeWidgetHelpers.h"
+#include "Tools/UEdMode.h"
+#include "LidarPointCloudEdMode.generated.h"
+
+namespace FLidarEditorModes
+{
+	LIDARPOINTCLOUDEDITOR_API extern const FEditorModeID EM_Lidar;
+}
+
+/**
+* Lidar editor mode
+*/
+UCLASS()
+class ULidarEditorMode : public UBaseLegacyWidgetEdMode
+{
+	FDelegateHandle PaletteChangedHandle;
+
+public:
+	GENERATED_BODY()
+	
+	ULidarEditorMode();
+
+	virtual void Enter() override;
+	virtual void Exit() override;
+	virtual bool IsSelectionAllowed(AActor* InActor, bool bInSelection) const override;
+	virtual void CreateToolkit() override;
+	virtual bool GetPivotForOrbit(FVector& OutPivot) const override;
+	virtual TMap<FName, TArray<TSharedPtr<FUICommandInfo>>> GetModeCommands() const override;
+	virtual void BindCommands() override;
+	virtual void OnToolStarted(UInteractiveToolManager* Manager, UInteractiveTool* Tool) override;
+	virtual void OnToolEnded(UInteractiveToolManager* Manager, UInteractiveTool* Tool) override;
+
+private:
+	void UpdateOnPaletteChange(FName NewPalette);
+	void CancelActiveToolAction();
+};
\ No newline at end of file
diff --git a/Engine/Plugins/Enterprise/LidarPointCloud/Source/LidarPointCloudEditor/Private/LidarPointCloudEdModeToolkit.cpp b/Engine/Plugins/Enterprise/LidarPointCloud/Source/LidarPointCloudEditor/Private/LidarPointCloudEdModeToolkit.cpp
new file mode 100644
--- /dev/null
+++ b/Engine/Plugins/Enterprise/LidarPointCloud/Source/LidarPointCloudEditor/Private/LidarPointCloudEdModeToolkit.cpp
@@ -0,0 +1,89 @@
+// Copyright Epic Games, Inc. All Rights Reserved.
+
+#include "LidarPointCloudEdModeToolkit.h"
+#include "EditorModeManager.h"
+#include "EdModeInteractiveToolsContext.h"
+#include "InteractiveToolManager.h"
+#include "LidarPointCloudEditorCommands.h"
+#include "LidarPointCloudEdMode.h"
+#include "LidarPointCloudShared.h"
+#include "Selection.h"
+#include "StatusBarSubsystem.h"
+#include "Toolkits/AssetEditorModeUILayer.h"
+
+#define LOCTEXT_NAMESPACE "LidarEditMode"
+
+FLidarPointCloudEdModeToolkit::~FLidarPointCloudEdModeToolkit()
+{
+	GetScriptableEditorMode()->GetInteractiveToolsContext(EToolsContextScope::EdMode)->OnToolNotificationMessage.RemoveAll(this);
+	GetScriptableEditorMode()->GetInteractiveToolsContext(EToolsContextScope::EdMode)->OnToolWarningMessage.RemoveAll(this);
+}
+
+void FLidarPointCloudEdModeToolkit::Init(const TSharedPtr<IToolkitHost>& InitToolkitHost, TWeakObjectPtr<UEdMode> InOwningMode)
+{
+	FModeToolkit::Init(InitToolkitHost, InOwningMode);
+	
+	GetScriptableEditorMode()->GetInteractiveToolsContext(EToolsContextScope::EdMode)->OnToolNotificationMessage.AddSP(this, &FLidarPointCloudEdModeToolkit::SetActiveToolMessage);
+	GetScriptableEditorMode()->GetInteractiveToolsContext(EToolsContextScope::EdMode)->OnToolWarningMessage.AddSP(this, &FLidarPointCloudEdModeToolkit::SetActiveToolMessage);
+}
+
+FName FLidarPointCloudEdModeToolkit::GetToolkitFName() const
+{
+	return FName("LidarEditMode");
+}
+
+FText FLidarPointCloudEdModeToolkit::GetBaseToolkitName() const
+{
+	return LOCTEXT( "ToolkitName", "Lidar" );
+}
+
+void FLidarPointCloudEdModeToolkit::GetToolPaletteNames(TArray<FName>& InPaletteName) const
+{
+	InPaletteName.Add(LidarEditorPalletes::Manage);
+	InPaletteName.Add(LidarEditorPalletes::Edit);
+}
+
+FText FLidarPointCloudEdModeToolkit::GetToolPaletteDisplayName(FName PaletteName) const
+{
+	if (PaletteName == LidarEditorPalletes::Manage)
+	{
+		return LOCTEXT("LidarMode_Manage", "Manage");
+	}
+	if (PaletteName == LidarEditorPalletes::Edit)
+	{
+		return LOCTEXT("LidarMode_Edit", "Edit");
+	}
+	return FText();
+}
+
+FText FLidarPointCloudEdModeToolkit::GetActiveToolDisplayName() const
+{
+	if (UInteractiveTool* ActiveTool = GetScriptableEditorMode()->GetToolManager()->GetActiveTool(EToolSide::Left))
+	{
+		return ActiveTool->GetClass()->GetDisplayNameText();
+	}
+
+	return LOCTEXT("LidarNoActiveTool", "LidarNoActiveTool");
+}
+
+FText FLidarPointCloudEdModeToolkit::GetActiveToolMessage() const
+{
+	return ActiveToolMessageCache;
+}
+
+void FLidarPointCloudEdModeToolkit::SetActiveToolMessage(const FText& Message)
+{
+	ActiveToolMessageCache = Message;
+	
+	if (ModeUILayer.IsValid())
+	{
+		TSharedPtr<FAssetEditorModeUILayer> ModeUILayerPtr = ModeUILayer.Pin();
+		GEditor->GetEditorSubsystem<UStatusBarSubsystem>()->PopStatusBarMessage(ModeUILayerPtr->GetStatusBarName(), ActiveToolMessageHandle);
+		ActiveToolMessageHandle = GEditor->GetEditorSubsystem<UStatusBarSubsystem>()->PushStatusBarMessage(ModeUILayerPtr->GetStatusBarName(), Message);
+	}
+
+	ActiveToolMessageHandle.Reset();
+}
+
+
+#undef LOCTEXT_NAMESPACE
diff --git a/Engine/Plugins/Enterprise/LidarPointCloud/Source/LidarPointCloudEditor/Private/LidarPointCloudEdModeToolkit.h b/Engine/Plugins/Enterprise/LidarPointCloud/Source/LidarPointCloudEditor/Private/LidarPointCloudEdModeToolkit.h
new file mode 100644
--- /dev/null
+++ b/Engine/Plugins/Enterprise/LidarPointCloud/Source/LidarPointCloudEditor/Private/LidarPointCloudEdModeToolkit.h
@@ -0,0 +1,40 @@
+// Copyright Epic Games, Inc. All Rights Reserved.
+
+#pragma once
+#include "Toolkits/BaseToolkit.h"
+#include "EditorModeManager.h"
+#include "StatusBarSubsystem.h"
+
+namespace LidarEditorPalletes
+{
+	static const FName Manage(TEXT("ToolMode_Manage")); 
+	static const FName Edit(TEXT("ToolMode_Edit"));
+}
+
+/**
+ * Public interface to Lidar Edit mode.
+ */
+class FLidarPointCloudEdModeToolkit : public FModeToolkit
+{
+public:
+	FLidarPointCloudEdModeToolkit(){}
+	~FLidarPointCloudEdModeToolkit();
+	
+	/** Initializes the Lidar mode toolkit */
+	virtual void Init(const TSharedPtr<IToolkitHost>& InitToolkitHost, TWeakObjectPtr<UEdMode> InOwningMode) override;
+
+	/** IToolkit interface */
+	virtual FName GetToolkitFName() const override;
+	virtual FText GetBaseToolkitName() const override;
+	virtual void GetToolPaletteNames(TArray<FName>& InPaletteName) const override;
+	virtual FText GetToolPaletteDisplayName(FName PaletteName) const override;
+	virtual bool HasIntegratedToolPalettes() const override { return true; }	
+	virtual FText GetActiveToolDisplayName() const override;
+	virtual FText GetActiveToolMessage() const override;
+
+	void SetActiveToolMessage(const FText& Message);
+
+private:
+	FStatusBarMessageHandle ActiveToolMessageHandle;
+	FText ActiveToolMessageCache;
+};
\ No newline at end of file
diff --git a/Engine/Plugins/Enterprise/LidarPointCloud/Source/LidarPointCloudEditor/Private/LidarPointCloudEditor.cpp b/Engine/Plugins/Enterprise/LidarPointCloud/Source/LidarPointCloudEditor/Private/LidarPointCloudEditor.cpp
--- a/Engine/Plugins/Enterprise/LidarPointCloud/Source/LidarPointCloudEditor/Private/LidarPointCloudEditor.cpp
+++ b/Engine/Plugins/Enterprise/LidarPointCloud/Source/LidarPointCloudEditor/Private/LidarPointCloudEditor.cpp
@@ -7,6 +7,7 @@
 #include "LidarPointCloudShared.h"
 #include "ILidarPointCloudEditorModule.h"
 #include "LidarPointCloudEditorCommands.h"
+#include "LidarPointCloudEditorHelper.h"
 
 #include "Misc/ScopedSlowTask.h"
 #include "Styling/AppStyle.h"
@@ -18,6 +19,7 @@
 #include "Framework/Commands/UICommandList.h"
 #include "IContentBrowserSingleton.h"
 #include "ContentBrowserModule.h"
+
 #include "AssetRegistry/AssetRegistryModule.h"
 #include "Misc/MessageDialog.h"
 
@@ -178,8 +180,6 @@ void FLidarPointCloudEditor::AddReferencedObjects(FReferenceCollector& Collector
 
 void FLidarPointCloudEditor::InitPointCloudEditor(const EToolkitMode::Type Mode, const TSharedPtr< class IToolkitHost >& InitToolkitHost, class ULidarPointCloud* InitPointCloud)
 {
-	FLidarPointCloudEditorCommands::Register();
-
 	PointCloudBeingEdited = InitPointCloud;
 
 	// Register for rebuilding events
@@ -760,7 +760,7 @@ void FLidarPointCloudEditor::Extract()
 		return;
 	}
 
-	ULidarPointCloud* NewPointCloud = CreateNewAsset();
+	ULidarPointCloud* NewPointCloud = FLidarPointCloudEditorHelper::CreateNewAsset();
 	if (NewPointCloud)
 	{
 		NewPointCloud->SetData(SelectedPoints);
@@ -777,7 +777,7 @@ void FLidarPointCloudEditor::ExtractCopy()
 		return;
 	}
 
-	ULidarPointCloud* NewPointCloud = CreateNewAsset();
+	ULidarPointCloud* NewPointCloud = FLidarPointCloudEditorHelper::CreateNewAsset();
 	if (NewPointCloud)
 	{
 		NewPointCloud->SetData(SelectedPoints);
@@ -947,43 +947,4 @@ TArray<FAssetData> FLidarPointCloudEditor::SelectAssets(const FText& Title)
 	return ContentBrowserModule.Get().CreateModalOpenAssetDialog(OpenAssetDialogConfig);
 }
 
-FString FLidarPointCloudEditor::GetSaveAsLocation()
-{
-	// Initialize SaveAssetDialog config
-	FSaveAssetDialogConfig SaveAssetDialogConfig;
-	SaveAssetDialogConfig.DialogTitleOverride = LOCTEXT("SelectExtractDestination", "Select Extract Destination");
-	SaveAssetDialogConfig.DefaultPath = "/Game";
-	SaveAssetDialogConfig.AssetClassNames.Emplace(*ULidarPointCloud::StaticClass()->GetName());
-	SaveAssetDialogConfig.ExistingAssetPolicy = ESaveAssetDialogExistingAssetPolicy::AllowButWarn;
-
-	FContentBrowserModule& ContentBrowserModule = FModuleManager::LoadModuleChecked<FContentBrowserModule>("ContentBrowser");
-	return ContentBrowserModule.Get().CreateModalSaveAssetDialog(SaveAssetDialogConfig);
-}
-
-ULidarPointCloud* FLidarPointCloudEditor::CreateNewAsset()
-{
-	ULidarPointCloud* NewPointCloud = nullptr;
-
-	FString SaveObjectPath = GetSaveAsLocation();
-	if (!SaveObjectPath.IsEmpty())
-	{
-		// Attempt to load existing asset first
-		NewPointCloud = FindObject<ULidarPointCloud>(nullptr, *SaveObjectPath);
-
-		// Proceed to creating a new asset, if needed
-		if (!NewPointCloud)
-		{
-			const FString PackageName = FPackageName::ObjectPathToPackageName(SaveObjectPath);
-			const FString ObjectName = FPackageName::ObjectPathToObjectName(SaveObjectPath);
-
-			NewPointCloud = NewObject<ULidarPointCloud>(CreatePackage(*PackageName), ULidarPointCloud::StaticClass(), FName(*ObjectName), EObjectFlags::RF_Public | EObjectFlags::RF_Standalone);
-
-			FAssetRegistryModule::AssetCreated(NewPointCloud);
-			NewPointCloud->MarkPackageDirty();
-		}		
-	}
-
-	return NewPointCloud;
-}
-
 #undef LOCTEXT_NAMESPACE
\ No newline at end of file
diff --git a/Engine/Plugins/Enterprise/LidarPointCloud/Source/LidarPointCloudEditor/Private/LidarPointCloudEditor.h b/Engine/Plugins/Enterprise/LidarPointCloud/Source/LidarPointCloudEditor/Private/LidarPointCloudEditor.h
--- a/Engine/Plugins/Enterprise/LidarPointCloud/Source/LidarPointCloudEditor/Private/LidarPointCloudEditor.h
+++ b/Engine/Plugins/Enterprise/LidarPointCloud/Source/LidarPointCloudEditor/Private/LidarPointCloudEditor.h
@@ -117,8 +117,6 @@ private:
 	TSharedRef<SWidget> GenerateSelectionMenuContent();
 
 	TArray<FAssetData> SelectAssets(const FText& Title);
-	FString GetSaveAsLocation();
-	ULidarPointCloud* CreateNewAsset();
 
 private:
 	/**	The tab ids for all the tabs used */
diff --git a/Engine/Plugins/Enterprise/LidarPointCloud/Source/LidarPointCloudEditor/Private/LidarPointCloudEditorCommands.cpp b/Engine/Plugins/Enterprise/LidarPointCloud/Source/LidarPointCloudEditor/Private/LidarPointCloudEditorCommands.cpp
--- a/Engine/Plugins/Enterprise/LidarPointCloud/Source/LidarPointCloudEditor/Private/LidarPointCloudEditorCommands.cpp
+++ b/Engine/Plugins/Enterprise/LidarPointCloud/Source/LidarPointCloudEditor/Private/LidarPointCloudEditorCommands.cpp
@@ -2,6 +2,8 @@
 
 #include "LidarPointCloudEditorCommands.h"
 
+#include "LidarPointCloudEdModeToolkit.h"
+
 #define LOCTEXT_NAMESPACE "LidarPointCloudEditor"
 
 //////////////////////////////////////////////////////////////////////////
@@ -13,35 +15,53 @@ void FLidarPointCloudEditorCommands::RegisterCommands()
 	UI_COMMAND(SetShowBounds, "Bounds", "Toggles display of the bounds of the point cloud.", EUserInterfaceActionType::ToggleButton, FInputChord());
 	UI_COMMAND(SetShowCollision, "Collision", "Toggles display of the collision of the point cloud.", EUserInterfaceActionType::ToggleButton, FInputChord());
 	UI_COMMAND(SetShowNodes, "Nodes", "Toggles display of the nodes of the point cloud.", EUserInterfaceActionType::ToggleButton, FInputChord());
-
 	UI_COMMAND(ResetCamera, "Reset Camera", "Resets the camera to focus on the point cloud.", EUserInterfaceActionType::Button, FInputChord());
-
 	UI_COMMAND(Center, "Center", "Enable, to center the point cloud asset\nDisable, to use original coordinates.", EUserInterfaceActionType::ToggleButton, FInputChord());
 	UI_COMMAND(BuildCollision, "Build Collision", "Builds collision for this point cloud.", EUserInterfaceActionType::Button, FInputChord());
 	UI_COMMAND(RemoveCollision, "Remove Collision", "Removes collision from this point cloud.", EUserInterfaceActionType::Button, FInputChord());
-
 	UI_COMMAND(EditMode, "Edit Mode", "Enables editing of the point cloud.", EUserInterfaceActionType::ToggleButton, FInputChord());
-
 	UI_COMMAND(BoxSelection, "Box Selection", "Uses box to select points.", EUserInterfaceActionType::Button, FInputChord());
 	UI_COMMAND(PolygonalSelection, "Polygonal Selection", "Uses custom polygon to select points.", EUserInterfaceActionType::Button, FInputChord());
 	UI_COMMAND(LassoSelection, "Lasso Selection", "Uses custom drawn shape to select points.", EUserInterfaceActionType::Button, FInputChord());
 	UI_COMMAND(PaintSelection, "Paint Selection", "Uses adjustable paint brush to select points.", EUserInterfaceActionType::Button, FInputChord());
-
 	UI_COMMAND(InvertSelection, "Invert Selection", "Inverts point selection.", EUserInterfaceActionType::Button, FInputChord());
-
 	UI_COMMAND(HideSelected, "Hide Selected", "Hide selected points.", EUserInterfaceActionType::Button, FInputChord());
 	UI_COMMAND(UnhideAll, "Unhide All", "Resets the visibility of all points in the point cloud.", EUserInterfaceActionType::Button, FInputChord());
-
 	UI_COMMAND(DeleteSelected, "Delete Selected", "Permanently remove selected points from the point cloud.", EUserInterfaceActionType::Button, FInputChord());
 	UI_COMMAND(DeleteHidden, "Delete Hidden", "Permanently remove hidden points from the point cloud.", EUserInterfaceActionType::Button, FInputChord());
-
 	UI_COMMAND(CalculateNormals, "Calculate Normals", "Calculates normals for the point cloud.", EUserInterfaceActionType::Button, FInputChord());
 	UI_COMMAND(CalculateNormalsSelection, "Calculate Normals (Selection)", "Calculates normals for the selected points.", EUserInterfaceActionType::Button, FInputChord());
-
 	UI_COMMAND(Extract, "Extract", "Extracts the selected points as a separate point cloud asset.", EUserInterfaceActionType::Button, FInputChord());
 	UI_COMMAND(ExtractCopy, "Extract as Copy", "Duplicates the selected points as a separate point cloud asset.", EUserInterfaceActionType::Button, FInputChord());
 	UI_COMMAND(Merge, "Merge", "Merges selected point cloud assets with this one.", EUserInterfaceActionType::Button, FInputChord());
 	UI_COMMAND(Align, "Align", "Aligns selected point cloud assets with this one while retaining overall centering.", EUserInterfaceActionType::Button, FInputChord());
+
+	// TOOLKIT COMMANDS
+
+	TArray<TSharedPtr<FUICommandInfo, ESPMode::ThreadSafe>>& ManageCommands = Commands.FindOrAdd(LidarEditorPalletes::Manage);
+	UI_COMMAND(ToolkitSelect, "Select", "Select point cloud assets", EUserInterfaceActionType::ToggleButton, FInputChord());
+	ManageCommands.Add(ToolkitSelect);
+	UI_COMMAND(ToolkitAlign, "Align", "Aligns selected point cloud assets", EUserInterfaceActionType::ToggleButton, FInputChord());
+	ManageCommands.Add(ToolkitAlign);
+	UI_COMMAND(ToolkitMerge, "Merge", "Merges selected point cloud assets", EUserInterfaceActionType::ToggleButton, FInputChord());
+	ManageCommands.Add(ToolkitMerge);
+	UI_COMMAND(ToolkitCollision, "Collision", "Create and Destroy point cloud asset collisions", EUserInterfaceActionType::ToggleButton, FInputChord());
+	ManageCommands.Add(ToolkitCollision);
+	UI_COMMAND(ToolkitNormals, "Normals", "Calculates normals", EUserInterfaceActionType::ToggleButton, FInputChord());
+	ManageCommands.Add(ToolkitNormals);
+	
+	TArray<TSharedPtr<FUICommandInfo, ESPMode::ThreadSafe>>& EditCommands = Commands.FindOrAdd(LidarEditorPalletes::Edit);
+	EditCommands.Add(ToolkitSelect);
+	UI_COMMAND(ToolkitBoxSelection, "Box", "Uses box to select points.", EUserInterfaceActionType::ToggleButton, FInputChord());
+	EditCommands.Add(ToolkitBoxSelection);
+	UI_COMMAND(ToolkitPolygonalSelection, "Polygon", "Uses custom polygon to select points.", EUserInterfaceActionType::ToggleButton, FInputChord());
+	EditCommands.Add(ToolkitPolygonalSelection);
+	UI_COMMAND(ToolkitLassoSelection, "Lasso", "Uses custom drawn shape to select points.", EUserInterfaceActionType::ToggleButton, FInputChord());
+	EditCommands.Add(ToolkitLassoSelection);
+	UI_COMMAND(ToolkitPaintSelection, "Paint", "Uses adjustable paint brush to select points.", EUserInterfaceActionType::ToggleButton, FInputChord());
+	EditCommands.Add(ToolkitPaintSelection);
+
+	UI_COMMAND(ToolktitCancelSelection, "Cancel Selection", "Cancel the active tool", EUserInterfaceActionType::Button, FInputChord(EKeys::Escape));
 }
 
 //////////////////////////////////////////////////////////////////////////
diff --git a/Engine/Plugins/Enterprise/LidarPointCloud/Source/LidarPointCloudEditor/Private/LidarPointCloudEditorCommands.h b/Engine/Plugins/Enterprise/LidarPointCloud/Source/LidarPointCloudEditor/Private/LidarPointCloudEditorCommands.h
--- a/Engine/Plugins/Enterprise/LidarPointCloud/Source/LidarPointCloudEditor/Private/LidarPointCloudEditorCommands.h
+++ b/Engine/Plugins/Enterprise/LidarPointCloud/Source/LidarPointCloudEditor/Private/LidarPointCloudEditorCommands.h
@@ -24,38 +24,50 @@ public:
 	virtual void RegisterCommands() override;
 	// End of TCommand<> interface
 
-public:
-	TSharedPtr< FUICommandInfo > SetShowGrid;
-	TSharedPtr< FUICommandInfo > SetShowBounds;
-	TSharedPtr< FUICommandInfo > SetShowCollision;
-	TSharedPtr< FUICommandInfo > SetShowNodes;
-
-	TSharedPtr< FUICommandInfo > ResetCamera;
-
-	TSharedPtr< FUICommandInfo > Center;
-	TSharedPtr< FUICommandInfo > BuildCollision;
-	TSharedPtr< FUICommandInfo > RemoveCollision;
-
-	TSharedPtr< FUICommandInfo > EditMode;
-
-	TSharedPtr< FUICommandInfo > BoxSelection;
-	TSharedPtr< FUICommandInfo > PolygonalSelection;
-	TSharedPtr< FUICommandInfo > LassoSelection;
-	TSharedPtr< FUICommandInfo > PaintSelection;
-
-	TSharedPtr< FUICommandInfo > InvertSelection;
+	static TMap<FName, TArray<TSharedPtr<FUICommandInfo>>> GetCommands()
+	{
+		return Get().Commands;
+	}
 
-	TSharedPtr< FUICommandInfo > HideSelected;
-	TSharedPtr< FUICommandInfo > UnhideAll;
+public:
+	TSharedPtr<FUICommandInfo> SetShowGrid;
+	TSharedPtr<FUICommandInfo> SetShowBounds;
+	TSharedPtr<FUICommandInfo> SetShowCollision;
+	TSharedPtr<FUICommandInfo> SetShowNodes;
+	TSharedPtr<FUICommandInfo> ResetCamera;
+	TSharedPtr<FUICommandInfo> Center;
+	TSharedPtr<FUICommandInfo> BuildCollision;
+	TSharedPtr<FUICommandInfo> RemoveCollision;
+	TSharedPtr<FUICommandInfo> EditMode;
+	TSharedPtr<FUICommandInfo> BoxSelection;
+	TSharedPtr<FUICommandInfo> PolygonalSelection;
+	TSharedPtr<FUICommandInfo> LassoSelection;
+	TSharedPtr<FUICommandInfo> PaintSelection;
+	TSharedPtr<FUICommandInfo> InvertSelection;
+	TSharedPtr<FUICommandInfo> HideSelected;
+	TSharedPtr<FUICommandInfo> UnhideAll;
+	TSharedPtr<FUICommandInfo> DeleteSelected;
+	TSharedPtr<FUICommandInfo> DeleteHidden;
+	TSharedPtr<FUICommandInfo> CalculateNormals;
+	TSharedPtr<FUICommandInfo> CalculateNormalsSelection;
+	TSharedPtr<FUICommandInfo> Extract;
+	TSharedPtr<FUICommandInfo> ExtractCopy;
+	TSharedPtr<FUICommandInfo> Merge;
+	TSharedPtr<FUICommandInfo> Align;
 
-	TSharedPtr< FUICommandInfo > DeleteSelected;
-	TSharedPtr< FUICommandInfo > DeleteHidden;
+	// TOOLKIT COMMANDS
 
-	TSharedPtr< FUICommandInfo > CalculateNormals;
-	TSharedPtr< FUICommandInfo > CalculateNormalsSelection;
+	TSharedPtr<FUICommandInfo> ToolkitSelect;
+	TSharedPtr<FUICommandInfo> ToolkitAlign;
+	TSharedPtr<FUICommandInfo> ToolkitMerge;
+	TSharedPtr<FUICommandInfo> ToolkitCollision;
+	TSharedPtr<FUICommandInfo> ToolkitNormals;
+	TSharedPtr<FUICommandInfo> ToolkitBoxSelection;
+	TSharedPtr<FUICommandInfo> ToolkitPolygonalSelection;
+	TSharedPtr<FUICommandInfo> ToolkitLassoSelection;
+	TSharedPtr<FUICommandInfo> ToolkitPaintSelection;
+	TSharedPtr<FUICommandInfo> ToolktitCancelSelection;
 
-	TSharedPtr< FUICommandInfo > Extract;
-	TSharedPtr< FUICommandInfo > ExtractCopy;
-	TSharedPtr< FUICommandInfo > Merge;
-	TSharedPtr< FUICommandInfo > Align;
+private:
+	TMap<FName, TArray<TSharedPtr<FUICommandInfo>>> Commands;
 };
\ No newline at end of file
diff --git a/Engine/Plugins/Enterprise/LidarPointCloud/Source/LidarPointCloudEditor/Private/LidarPointCloudEditorHelper.cpp b/Engine/Plugins/Enterprise/LidarPointCloud/Source/LidarPointCloudEditor/Private/LidarPointCloudEditorHelper.cpp
new file mode 100644
--- /dev/null
+++ b/Engine/Plugins/Enterprise/LidarPointCloud/Source/LidarPointCloudEditor/Private/LidarPointCloudEditorHelper.cpp
@@ -0,0 +1,702 @@
+// Copyright Epic Games, Inc. All Rights Reserved.
+
+#include "LidarPointCloudEditorHelper.h"
+#include "ContentBrowserModule.h"
+#include "GeomTools.h"
+#include "IContentBrowserSingleton.h"
+#include "LevelEditorViewport.h"
+#include "LidarPointCloud.h"
+#include "LidarPointCloudActor.h"
+#include "LidarPointCloudComponent.h"
+#include "Selection.h"
+
+#include "AssetRegistry/AssetRegistryModule.h"
+
+#include "Misc/ScopedSlowTask.h"
+
+#define LOCTEXT_NAMESPACE "LidarPointCloudEditorHelper"
+
+class FContentBrowserModule;
+
+namespace
+{
+	FString GetSaveAsLocation()
+	{
+		// Initialize SaveAssetDialog config
+		FSaveAssetDialogConfig SaveAssetDialogConfig;
+		SaveAssetDialogConfig.DialogTitleOverride = LOCTEXT("SelectDestination", "Select Destination");
+		SaveAssetDialogConfig.DefaultPath = "/Game";
+		SaveAssetDialogConfig.AssetClassNames.Emplace("/Script/PointCloudRuntime.LidarPointCloud");
+		SaveAssetDialogConfig.ExistingAssetPolicy = ESaveAssetDialogExistingAssetPolicy::AllowButWarn;
+
+		FContentBrowserModule& ContentBrowserModule = FModuleManager::LoadModuleChecked<FContentBrowserModule>("ContentBrowser");
+		return ContentBrowserModule.Get().CreateModalSaveAssetDialog(SaveAssetDialogConfig);
+	}
+
+	ULidarPointCloud* CreateNewAsset_Internal()
+	{
+		ULidarPointCloud* NewPointCloud = nullptr;
+
+		const FString SaveObjectPath = GetSaveAsLocation();
+		if (!SaveObjectPath.IsEmpty())
+		{
+			// Attempt to load existing asset first
+			NewPointCloud = FindObject<ULidarPointCloud>(nullptr, *SaveObjectPath);
+
+			// Proceed to creating a new asset, if needed
+			if (!NewPointCloud)
+			{
+				const FString PackageName = FPackageName::ObjectPathToPackageName(SaveObjectPath);
+				const FString ObjectName = FPackageName::ObjectPathToObjectName(SaveObjectPath);
+
+				NewPointCloud = NewObject<ULidarPointCloud>(CreatePackage(*PackageName), ULidarPointCloud::StaticClass(), FName(*ObjectName), EObjectFlags::RF_Public | EObjectFlags::RF_Standalone);
+
+				FAssetRegistryModule::AssetCreated(NewPointCloud);
+				NewPointCloud->MarkPackageDirty();
+			}		
+		}
+
+		return NewPointCloud;
+	}
+
+	void ProcessSelection(TFunction<void(ALidarPointCloudActor*)> Function)
+	{
+		if(!Function)
+		{
+			return;
+		}
+		
+		for (FSelectionIterator It(GEditor->GetSelectedActorIterator()); It; ++It)
+		{
+			if(ALidarPointCloudActor* LidarActor = Cast<ALidarPointCloudActor>(*It))
+			{
+				Function(LidarActor);
+			}
+		}
+	}
+
+	void ProcessAll(TFunction<void(ALidarPointCloudActor*)> Function)
+	{
+		if(!Function)
+		{
+			return;
+		}
+		
+		for (TObjectIterator<ALidarPointCloudActor> It; It; ++It)
+		{
+			ALidarPointCloudActor* LidarActor = Cast<ALidarPointCloudActor>(*It);
+			if(IsValid(LidarActor))
+			{
+				Function(LidarActor);
+			}
+		}
+	}
+
+	TArray<ALidarPointCloudActor*> GetSelectedActors()
+	{
+		TArray<ALidarPointCloudActor*> Actors;
+		
+		for (FSelectionIterator It(GEditor->GetSelectedActorIterator()); It; ++It)
+		{
+			if(ALidarPointCloudActor* LidarActor = Cast<ALidarPointCloudActor>(*It))
+			{
+				Actors.Add(LidarActor);
+			}
+		}
+
+		return Actors;
+	}
+
+	UWorld* GetFirstWorld()
+	{
+		for (TObjectIterator<ALidarPointCloudActor> It; It; ++It)
+		{
+			ALidarPointCloudActor* LidarActor = Cast<ALidarPointCloudActor>(*It);
+			if(IsValid(LidarActor))
+			{
+				return LidarActor->GetWorld();
+			}
+		}
+		return nullptr;
+	}
+
+	TArray<ULidarPointCloud*> GetSelectedClouds()
+	{
+		TArray<ULidarPointCloud*> PointClouds;
+		
+		for (FSelectionIterator It(GEditor->GetSelectedActorIterator()); It; ++It)
+		{
+			if(ALidarPointCloudActor* LidarActor = Cast<ALidarPointCloudActor>(*It))
+			{
+				if(ULidarPointCloud* PointCloud = LidarActor->GetPointCloud())
+				{
+					PointClouds.Add(PointCloud);
+				}
+			}
+		}
+
+		return PointClouds;
+	}
+
+	ALidarPointCloudActor* SpawnActor()
+	{
+		ALidarPointCloudActor* NewActor = nullptr;
+		if(UWorld* World = GetFirstWorld())
+		{
+			NewActor = Cast<ALidarPointCloudActor>(World->SpawnActor(ALidarPointCloudActor::StaticClass()));
+		}
+		return NewActor;
+	}
+	
+	FSceneView* GetEditorView(FEditorViewportClient* ViewportClient)
+	{
+		FSceneViewFamilyContext ViewFamily(FSceneViewFamily::ConstructionValues(
+			ViewportClient->Viewport,
+			ViewportClient->GetScene(),
+			ViewportClient->EngineShowFlags).SetRealtimeUpdate(ViewportClient->IsRealtime()));
+		return ViewportClient->CalcSceneView(&ViewFamily);
+	}
+
+	FConvexVolume BuildConvexVolumeForPoints(const TArray<FVector2D>& Points, FEditorViewportClient* ViewportClient)
+	{
+		if(!ViewportClient)
+		{
+			ViewportClient = GCurrentLevelEditingViewportClient;
+		}
+		
+		FConvexVolume ConvexVolume;
+
+		const FSceneView* View = GetEditorView(ViewportClient);
+		const FMatrix InvViewProjectionMatrix = View->ViewMatrices.GetInvViewProjectionMatrix();
+
+		TArray<FVector> Origins; Origins.AddUninitialized(Points.Num() + 2);
+		TArray<FVector> Normals; Normals.AddUninitialized(Points.Num() + 2);
+		TArray<FVector> Directions; Directions.AddUninitialized(Points.Num());
+		FVector MeanCenter = FVector::ZeroVector;
+
+		for (int32 i = 0; i < Points.Num(); ++i)
+		{
+			FSceneView::DeprojectScreenToWorld(Points[i], FIntRect(FIntPoint(0, 0), ViewportClient->Viewport->GetSizeXY()), InvViewProjectionMatrix, Origins[i], Directions[i]);
+			MeanCenter += Origins[i];
+		}
+
+		MeanCenter /= Points.Num();
+
+		const FVector& ViewDirection = View->GetViewDirection();
+
+		// Shared calculations
+		Normals.Last(1) = ViewDirection;
+		Normals.Last() = -ViewDirection;
+		Origins.Last(1) = Origins[0] + ViewDirection * 99999999.0f;
+
+		// Calculate plane normals
+		const bool bFlipNormals = FVector::DotProduct(Normals[0], (MeanCenter - Origins[0])) > 0;
+		for (int32 i = 0; i < Points.Num(); ++i)
+		{
+			Normals[i] = ((Origins[(i + 1) % Points.Num()] - Origins[i]).GetSafeNormal() ^ Directions[i]).GetSafeNormal();
+			
+			if (bFlipNormals)
+			{
+				Normals[i] = -Normals[i];
+			}
+		}
+
+		// Perspective View
+		if (View->IsPerspectiveProjection())
+		{
+			Origins.Last() = Origins[0];
+		}
+		// Ortho Views
+		else
+		{
+			Origins.Last() = -Origins.Last(1);
+		}
+
+		for (int32 i = 0; i < Origins.Num(); ++i)
+		{
+			ConvexVolume.Planes.Emplace(Origins[i], Normals[i]);
+		}
+
+		ConvexVolume.Init();
+
+		return ConvexVolume;
+	}
+	
+	ULidarPointCloud* Extract_Internal()
+	{
+		int64 NumPoints = 0;
+		ProcessAll([&NumPoints](ALidarPointCloudActor* Actor)
+		{
+			NumPoints += Actor->GetPointCloudComponent()->NumSelectedPoints();
+		});
+
+		TArray64<FLidarPointCloudPoint> SelectedPoints;
+		SelectedPoints.Reserve(NumPoints);
+		ProcessAll([&SelectedPoints](ALidarPointCloudActor* Actor)
+		{
+			Actor->GetPointCloudComponent()->GetSelectedPointsAsCopies(SelectedPoints);
+		});
+	
+		ULidarPointCloud* NewPointCloud = CreateNewAsset_Internal();
+		if (NewPointCloud)
+		{
+			NewPointCloud->SetData(SelectedPoints);
+		}
+		return NewPointCloud;
+	}
+	
+	// Copied from GeomTools.cpp
+	bool IsPolygonConvex(const TArray<FVector2D>& Points)
+	{
+		const int PointCount = Points.Num();
+		float Sign = 0;
+		for (int32 PointIndex = 0; PointIndex < PointCount; ++PointIndex)
+		{
+			const FVector2D& A = Points[PointIndex];
+			const FVector2D& B = Points[(PointIndex + 1) % PointCount];
+			const FVector2D& C = Points[(PointIndex + 2) % PointCount];
+			float Det = (B.X - A.X) * (C.Y - B.Y) - (B.Y - A.Y) * (C.X - B.X);
+			float DetSign = FMath::Sign(Det);
+			if (DetSign != 0)
+			{
+				if (Sign == 0)
+				{
+					Sign = DetSign;
+				}
+				else if (Sign != DetSign)
+				{
+					return false;
+				}
+			}
+		}
+
+		return true;
+	}
+}
+
+ULidarPointCloud* FLidarPointCloudEditorHelper::CreateNewAsset()
+{
+	return CreateNewAsset_Internal();
+}
+
+void FLidarPointCloudEditorHelper::AlignSelectionAroundWorldOrigin()
+{
+	ULidarPointCloud::AlignClouds(GetSelectedClouds());
+}
+
+void FLidarPointCloudEditorHelper::SetOriginalCoordinateForSelection()
+{
+	ProcessSelection([](ALidarPointCloudActor* Actor)
+	{
+		if(ULidarPointCloud* PointCloud = Actor->GetPointCloud())
+		{
+			PointCloud->RestoreOriginalCoordinates();
+		}
+	});
+}
+
+void FLidarPointCloudEditorHelper::CenterSelection()
+{
+	ProcessSelection([](ALidarPointCloudActor* Actor)
+	{
+		if(ULidarPointCloud* PointCloud = Actor->GetPointCloud())
+		{
+			PointCloud->CenterPoints();
+		}
+	});
+}
+
+void FLidarPointCloudEditorHelper::BuildCollisionForSelection()
+{
+	ProcessSelection([](ALidarPointCloudActor* Actor)
+	{
+		if(ULidarPointCloud* PointCloud = Actor->GetPointCloud())
+		{
+			PointCloud->BuildCollision();
+		}
+	});
+}
+
+void FLidarPointCloudEditorHelper::SetCollisionErrorForSelection(float Error)
+{
+	ProcessSelection([Error](ALidarPointCloudActor* Actor)
+	{
+		if(ULidarPointCloud* PointCloud = Actor->GetPointCloud())
+		{
+			if(Error > 0)
+			{
+				PointCloud->MaxCollisionError = Error;
+			}
+			else
+			{
+				PointCloud->SetOptimalCollisionError();
+			}
+		}
+	});
+}
+
+void FLidarPointCloudEditorHelper::RemoveCollisionForSelection()
+{
+	ProcessSelection([](ALidarPointCloudActor* Actor)
+	{
+		if(ULidarPointCloud* PointCloud = Actor->GetPointCloud())
+		{
+			PointCloud->RemoveCollision();
+		}
+	});
+}
+
+void FLidarPointCloudEditorHelper::CalculateNormalsForSelection()
+{
+	ProcessSelection([](ALidarPointCloudActor* Actor)
+	{
+		if(ULidarPointCloud* PointCloud = Actor->GetPointCloud())
+		{
+			PointCloud->CalculateNormals(nullptr, nullptr);
+		}
+	});
+}
+
+void FLidarPointCloudEditorHelper::SetNormalsQualityForSelection(int32 Quality, float NoiseTolerance)
+{
+	ProcessSelection([Quality, NoiseTolerance](ALidarPointCloudActor* Actor)
+	{
+		if(ULidarPointCloud* PointCloud = Actor->GetPointCloud())
+		{
+			PointCloud->NormalsQuality = Quality;
+			PointCloud->NormalsNoiseTolerance = NoiseTolerance;
+		}
+	});
+}
+
+void FLidarPointCloudEditorHelper::ResetVisibility()
+{
+	ProcessAll([](ALidarPointCloudActor* Actor)
+	{
+		if(ULidarPointCloud* PointCloud = Actor->GetPointCloud())
+		{
+			PointCloud->UnhideAll();
+		}
+	});
+}
+
+void FLidarPointCloudEditorHelper::DeleteHidden()
+{
+	ProcessAll([](ALidarPointCloudActor* Actor)
+	{
+		if(ULidarPointCloud* PointCloud = Actor->GetPointCloud())
+		{
+			PointCloud->RemoveHiddenPoints();
+		}
+	});
+}
+
+void FLidarPointCloudEditorHelper::Extract()
+{
+	ULidarPointCloud* NewPointCloud = Extract_Internal();
+	if(ALidarPointCloudActor* Actor = SpawnActor())
+	{
+		NewPointCloud->RestoreOriginalCoordinates();
+		Actor->SetPointCloud(NewPointCloud);
+	}
+
+	DeleteSelected();
+}
+
+void FLidarPointCloudEditorHelper::ExtractAsCopy()
+{
+	ULidarPointCloud* NewPointCloud = Extract_Internal();
+	if(ALidarPointCloudActor* Actor = SpawnActor())
+	{
+		NewPointCloud->RestoreOriginalCoordinates();
+		Actor->SetPointCloud(NewPointCloud);
+	}
+	
+	ClearSelection();
+}
+
+void FLidarPointCloudEditorHelper::CalculateNormals()
+{
+	ProcessAll([](ALidarPointCloudActor* Actor)
+	{
+		if(ULidarPointCloud* PointCloud = Actor->GetPointCloud())
+		{
+			PointCloud->CalculateNormalsForSelection();
+		}
+	});
+}
+
+FConvexVolume FLidarPointCloudEditorHelper::BuildConvexVolumeFromCoordinates(FVector2d Start, FVector2d End, FEditorViewportClient* ViewportClient)
+{
+	FIntVector4 SelectionArea;
+	SelectionArea.X = FMath::Min(Start.X, End.X);
+	SelectionArea.Y = FMath::Min(Start.Y, End.Y);
+	SelectionArea.Z = FMath::Max(Start.X, End.X);
+	SelectionArea.W = FMath::Max(Start.Y, End.Y);
+
+	return BuildConvexVolumeForPoints(TArray<FVector2D>({
+			FVector2D(SelectionArea.X, SelectionArea.Y),
+			FVector2D(SelectionArea.X, SelectionArea.W),
+			FVector2D(SelectionArea.Z, SelectionArea.W),
+			FVector2D(SelectionArea.Z, SelectionArea.Y) }), ViewportClient);
+}
+
+TArray<FConvexVolume> FLidarPointCloudEditorHelper::BuildConvexVolumesFromPoints(TArray<FVector2d> Points, FEditorViewportClient* ViewportClient)
+{
+	TArray<FConvexVolume> ConvexVolumes;
+		
+	if (IsPolygonConvex(Points))
+	{
+		ConvexVolumes.Add(BuildConvexVolumeForPoints(Points, ViewportClient));
+	}
+	else
+	{
+		// Check for self-intersecting shape
+		if (!IsPolygonSelfIntersecting(Points, true))
+		{
+			TArray<TArray<FVector2D>> ConvexShapes;
+			
+			// The separation needs points in CCW order
+			if (!FGeomTools2D::IsPolygonWindingCCW(Points))
+			{
+				Algo::Reverse(Points);
+			}
+
+			TArray<FVector2D> Triangles;
+			FGeomTools2D::TriangulatePoly(Triangles, Points, false);
+			FGeomTools2D::GenerateConvexPolygonsFromTriangles(ConvexShapes, Triangles);
+
+			for (int32 i = 0; i < ConvexShapes.Num(); ++i)
+			{
+				ConvexVolumes.Add(BuildConvexVolumeForPoints(ConvexShapes[i], ViewportClient));
+			}
+		}
+	}
+
+	return ConvexVolumes;
+}
+
+FLidarPointCloudRay FLidarPointCloudEditorHelper::MakeRayFromScreenPosition(FVector2d Position, FEditorViewportClient* ViewportClient)
+{
+	if(!ViewportClient)
+	{
+		ViewportClient = GCurrentLevelEditingViewportClient;
+	}
+	
+	const FSceneView* View = GetEditorView(ViewportClient);
+	const FMatrix InvViewProjectionMatrix = View->ViewMatrices.GetInvViewProjectionMatrix();
+
+	FVector3d Origin, Direction;
+	
+	FSceneView::DeprojectScreenToWorld(Position, FIntRect(FIntPoint(0, 0), ViewportClient->Viewport->GetSizeXY()), InvViewProjectionMatrix, Origin, Direction);
+
+	return FLidarPointCloudRay(Origin, Direction);
+}
+
+bool FLidarPointCloudEditorHelper::RayTracePointClouds(const FLidarPointCloudRay& Ray, float RadiusMulti, FVector3f& OutHitLocation)
+{
+	float MinDistance = FLT_MAX;
+	
+	ProcessAll([Ray, &MinDistance, &OutHitLocation, RadiusMulti](ALidarPointCloudActor* Actor)
+	{
+		if(ULidarPointCloudComponent* Component = Actor->GetPointCloudComponent())
+		{
+			if(const ULidarPointCloud* PointCloud = Component->GetPointCloud())
+			{
+				const float TraceRadius = FMath::Max(PointCloud->GetEstimatedPointSpacing(), 0.5f) * RadiusMulti;
+				if(const FLidarPointCloudPoint* Point = Component->LineTraceSingle(Ray, TraceRadius, true))
+				{
+					const FVector3f PointLocation = (FVector3f)(Component->GetComponentTransform().TransformPosition((FVector)Point->Location) + PointCloud->LocationOffset);
+					
+					const float DistanceSq = (Ray.Origin - PointLocation).SizeSquared();
+					if(DistanceSq < MinDistance)
+					{
+						MinDistance = DistanceSq;
+						OutHitLocation = PointLocation;
+					}
+				}
+			}
+		}
+	});
+
+	return MinDistance < FLT_MAX;
+}
+
+bool FLidarPointCloudEditorHelper::IsPolygonSelfIntersecting(const TArray<FVector2D>& Points, bool bAllowLooping)
+{
+	// Slow, O(n2), but sufficient for the current problem
+	
+	const int32 MaxIndex = bAllowLooping ? Points.Num() : Points.Num() - 1;
+
+	for (int32 i = 0; i < MaxIndex; ++i)
+	{
+		const int32 i1 = (i + 1) % Points.Num();
+
+		const FVector2D P1 = Points[i];
+		const FVector2D P2 = Points[i1];
+
+		for (int32 j = 0; j < MaxIndex; ++j)
+		{
+			const int32 j1 = j < Points.Num() - 1 ? j + 1 : 0;
+
+			if (j1 != i && j != i && j != i1)
+			{
+				// Modified FMath::SegmentIntersection2D
+				// Inlining the code and skipping calculation of an intersection point makes a slight difference for O(n2)
+				const FVector2D SegmentStartA = P1;
+				const FVector2D SegmentEndA = P2;
+				const FVector2D SegmentStartB = Points[j];
+				const FVector2D SegmentEndB = Points[j1];
+				const FVector2D VectorA = P2 - SegmentStartA;
+				const FVector2D VectorB = SegmentEndB - SegmentStartB;
+
+				const float S = (-VectorA.Y * (SegmentStartA.X - SegmentStartB.X) + VectorA.X * (SegmentStartA.Y - SegmentStartB.Y)) / (-VectorB.X * VectorA.Y + VectorA.X * VectorB.Y);
+				if (S >= 0 && S <= 1)
+				{
+					const float T = (VectorB.X * (SegmentStartA.Y - SegmentStartB.Y) - VectorB.Y * (SegmentStartA.X - SegmentStartB.X)) / (-VectorB.X * VectorA.Y + VectorA.X * VectorB.Y);
+					if (T >= 0 && T <= 1)
+					{
+						return true;
+					}
+				}
+			}
+		}
+	}
+
+	return false;
+}
+
+void FLidarPointCloudEditorHelper::SelectPointsByConvexVolume(const FConvexVolume& ConvexVolume, ELidarPointCloudSelectionMode SelectionMode)
+{
+	ProcessAll([ConvexVolume, SelectionMode](ALidarPointCloudActor* Actor)
+	{
+		if(ULidarPointCloudComponent* Component = Actor->GetPointCloudComponent())
+		{
+			Component->SelectByConvexVolume(ConvexVolume, SelectionMode != ELidarPointCloudSelectionMode::Subtract, true);
+		}
+	});
+}
+
+void FLidarPointCloudEditorHelper::SelectPointsBySphere(FSphere Sphere, ELidarPointCloudSelectionMode SelectionMode)
+{
+	ProcessAll([Sphere, SelectionMode](ALidarPointCloudActor* Actor)
+	{
+		if(ULidarPointCloudComponent* Component = Actor->GetPointCloudComponent())
+		{
+			Component->SelectBySphere(Sphere, SelectionMode != ELidarPointCloudSelectionMode::Subtract, true);
+		}
+	});
+}
+
+void FLidarPointCloudEditorHelper::HideSelected()
+{
+	ProcessAll([](ALidarPointCloudActor* Actor)
+	{
+		Actor->GetPointCloudComponent()->HideSelected();
+	});
+}
+
+void FLidarPointCloudEditorHelper::DeleteSelected()
+{
+	ProcessAll([](ALidarPointCloudActor* Actor)
+	{
+		Actor->GetPointCloudComponent()->DeleteSelected();
+	});
+}
+
+void FLidarPointCloudEditorHelper::InvertSelection()
+{
+	ProcessAll([](ALidarPointCloudActor* Actor)
+	{
+		Actor->GetPointCloudComponent()->InvertSelection();
+	});
+}
+
+void FLidarPointCloudEditorHelper::ClearSelection()
+{
+	ProcessAll([](ALidarPointCloudActor* Actor)
+	{
+		Actor->GetPointCloudComponent()->ClearSelection();
+	});
+}
+
+void FLidarPointCloudEditorHelper::MergeLidar(ULidarPointCloud* TargetAsset, TArray<ULidarPointCloud*> SourceAssets)
+{
+	if(!IsValid(TargetAsset) || SourceAssets.Num() == 0)
+	{
+		return;
+	}
+	
+	FScopedSlowTask ProgressDialog(SourceAssets.Num() + 2, LOCTEXT("Merge", "Merging Point Clouds..."));
+	ProgressDialog.MakeDialog();
+
+	TargetAsset->Merge(SourceAssets, [&ProgressDialog]() { ProgressDialog.EnterProgressFrame(1.f); });
+
+	FAssetRegistryModule::AssetCreated(TargetAsset);
+	TargetAsset->MarkPackageDirty();
+}
+
+void FLidarPointCloudEditorHelper::MergeSelectionByData(bool bReplaceSource)
+{
+	TArray<ALidarPointCloudActor*> Actors = GetSelectedActors();
+	TArray<ULidarPointCloud*> PointClouds = GetSelectedClouds();
+	
+	if (PointClouds.Num() > 1)
+	{
+		ULidarPointCloud* NewCloud = CreateNewAsset();
+		
+		MergeLidar(NewCloud, PointClouds);
+
+		if(bReplaceSource)
+		{
+			// Repurpose the first actor
+			Actors[0]->SetPointCloud(NewCloud);
+			
+			// Remove the rest
+			for(int32 i = 1; i < Actors.Num(); ++i)
+			{
+				Actors[i]->Destroy();
+			}
+		}
+		else
+		{
+			ALidarPointCloudActor* NewActor = Actors[0]->GetWorld()->SpawnActor<ALidarPointCloudActor>(Actors[0]->GetActorLocation(), Actors[0]->GetActorRotation());
+			NewActor->SetPointCloud(NewCloud);
+		}
+	}
+}
+
+void FLidarPointCloudEditorHelper::MergeSelectionByComponent(bool bReplaceSource)
+{
+	TArray<ALidarPointCloudActor*> Actors = GetSelectedActors();
+
+	if (Actors.Num() > 1)
+	{
+		AActor* TargetActor = Actors[0]->GetWorld()->SpawnActor<ALidarPointCloudActor>();
+		
+		for(ALidarPointCloudActor* Actor : Actors)
+		{
+			TArray<ULidarPointCloudComponent*> Components;
+			Actor->GetComponents(Components);
+
+			for(const ULidarPointCloudComponent* Component : Components)
+			{
+				ULidarPointCloudComponent* PCC = (ULidarPointCloudComponent*)TargetActor->AddComponentByClass(ULidarPointCloudComponent::StaticClass(), true, Component->GetComponentTransform(), false);
+				PCC->SetPointCloud(Component->GetPointCloud());
+				PCC->SetWorldTransform(Component->GetComponentTransform());
+				PCC->AttachToComponent(TargetActor->GetRootComponent(), FAttachmentTransformRules(EAttachmentRule::KeepWorld, false));
+			}
+
+			if(bReplaceSource)
+			{
+				Actor->Destroy();
+			}
+			else
+			{
+				Actor->SetHidden(true);
+			}
+		}
+	}
+}
+
+#undef LOCTEXT_NAMESPACE
diff --git a/Engine/Plugins/Enterprise/LidarPointCloud/Source/LidarPointCloudEditor/Private/LidarPointCloudEditorHelper.h b/Engine/Plugins/Enterprise/LidarPointCloud/Source/LidarPointCloudEditor/Private/LidarPointCloudEditorHelper.h
new file mode 100644
--- /dev/null
+++ b/Engine/Plugins/Enterprise/LidarPointCloud/Source/LidarPointCloudEditor/Private/LidarPointCloudEditorHelper.h
@@ -0,0 +1,58 @@
+// Copyright Epic Games, Inc. All Rights Reserved.
+
+#pragma once
+
+#include "CoreMinimal.h"
+#include "LidarPointCloudShared.h"
+
+class ALidarPointCloudActor;
+class ULidarPointCloud;
+
+enum class ELidarPointCloudSelectionMode : uint8
+{
+	None,
+	Add,
+	Subtract
+};
+
+class FLidarPointCloudEditorHelper
+{
+public:
+	static ULidarPointCloud* CreateNewAsset();
+
+	static void AlignSelectionAroundWorldOrigin();
+	static void SetOriginalCoordinateForSelection();
+	static void CenterSelection();
+	
+	static void MergeLidar(ULidarPointCloud* TargetAsset, TArray<ULidarPointCloud*> SourceAssets);
+	static void MergeSelectionByData(bool bReplaceSource);
+	static void MergeSelectionByComponent(bool bReplaceSource);
+
+	static void BuildCollisionForSelection();
+	static void SetCollisionErrorForSelection(float Error);
+	static void RemoveCollisionForSelection();
+	
+	static void CalculateNormalsForSelection();
+	static void SetNormalsQualityForSelection(int32 Quality, float NoiseTolerance);
+	
+	static void SelectPointsByConvexVolume(const FConvexVolume& ConvexVolume, ELidarPointCloudSelectionMode SelectionMode);
+	static void SelectPointsBySphere(FSphere Sphere, ELidarPointCloudSelectionMode SelectionMode);
+	static void HideSelected();
+	static void DeleteSelected();
+	static void InvertSelection();
+	static void ClearSelection();
+
+	static void ResetVisibility();
+	static void DeleteHidden();
+	static void Extract();
+	static void ExtractAsCopy();
+	static void CalculateNormals();
+
+	static FConvexVolume BuildConvexVolumeFromCoordinates(FVector2d Start, FVector2d End, FEditorViewportClient* ViewportClient = nullptr);
+	static TArray<FConvexVolume> BuildConvexVolumesFromPoints(TArray<FVector2d> Points, FEditorViewportClient* ViewportClient = nullptr);
+	static FLidarPointCloudRay MakeRayFromScreenPosition(FVector2d Position, FEditorViewportClient* ViewportClient = nullptr);
+
+	static bool RayTracePointClouds(const FLidarPointCloudRay& Ray, float RadiusMulti, FVector3f& OutHitLocation);
+
+	static bool IsPolygonSelfIntersecting(const TArray<FVector2D>& Points, bool bAllowLooping);
+};
diff --git a/Engine/Plugins/Enterprise/LidarPointCloud/Source/LidarPointCloudEditor/Private/LidarPointCloudEditorModule.cpp b/Engine/Plugins/Enterprise/LidarPointCloud/Source/LidarPointCloudEditor/Private/LidarPointCloudEditorModule.cpp
--- a/Engine/Plugins/Enterprise/LidarPointCloud/Source/LidarPointCloudEditor/Private/LidarPointCloudEditorModule.cpp
+++ b/Engine/Plugins/Enterprise/LidarPointCloud/Source/LidarPointCloudEditor/Private/LidarPointCloudEditorModule.cpp
@@ -1,5 +1,6 @@
 // Copyright Epic Games, Inc. All Rights Reserved.
 
+#include "EditorModeRegistry.h"
 #include "ILidarPointCloudEditorModule.h"
 #include "Modules/ModuleManager.h"
 #include "PropertyEditorModule.h"
@@ -10,6 +11,8 @@
 #include "LidarPointCloudShared.h"
 #include "LidarPointCloudSettings.h"
 #include "ISettingsModule.h"
+#include "LidarPointCloudEditorCommands.h"
+#include "LidarPointCloudEdMode.h"
 #include "ShowFlags.h"
 
 DEFINE_LOG_CATEGORY(LogLidarPointCloud);
@@ -44,12 +47,18 @@ class FLidarPointCloudEditorModule : public ILidarPointCloudEditorModule
 		}
 
 		FEngineShowFlags::RegisterCustomShowFlag(TEXT("LidarClippingVolume"), true, EShowFlagGroup::SFG_Transient, NSLOCTEXT("UnrealEd", "LidarClippingVolumeSF", "Lidar Clipping Volume"));
+
+		FLidarPointCloudEditorCommands::Register();
 	}
 	virtual void ShutdownModule() override
 	{
 		MenuExtensibilityManager.Reset();
 		ToolBarExtensibilityManager.Reset();
 
+		FLidarPointCloudEditorCommands::Unregister();
+
+		FEditorModeRegistry::Get().UnregisterMode(FLidarEditorModes::EM_Lidar);
+
 		if (UObjectInitialized())
 		{
 			// Unregister settings
diff --git a/Engine/Plugins/Enterprise/LidarPointCloud/Source/LidarPointCloudEditor/Private/LidarPointCloudEditorTools.cpp b/Engine/Plugins/Enterprise/LidarPointCloud/Source/LidarPointCloudEditor/Private/LidarPointCloudEditorTools.cpp
new file mode 100644
--- /dev/null
+++ b/Engine/Plugins/Enterprise/LidarPointCloud/Source/LidarPointCloudEditor/Private/LidarPointCloudEditorTools.cpp
@@ -0,0 +1,459 @@
+// Copyright Epic Games, Inc. All Rights Reserved.
+
+#include "LidarPointCloudEditorTools.h"
+
+#include "CanvasItem.h"
+#include "CanvasTypes.h"
+#include "LidarPointCloudEditorHelper.h"
+#include "LidarPointCloudShared.h"
+#include "BaseBehaviors/ClickDragBehavior.h"
+#include "BaseBehaviors/MouseHoverBehavior.h"
+
+#define LOCTEXT_NAMESPACE "LidarPointCloudEditorTool"
+
+namespace {
+	// Distance Square between the first and last points of the polygonal selection, where the shape will be considered as closed
+	constexpr int32 PolySnapDistanceSq = 250;
+
+	// Affects the frequency of new point injections when drawing lasso-based shapes
+	constexpr int32 LassoSpacingSq = 250;
+
+	// Affects the max depth delta when painting. Prevents the brush from "falling through" the gaps.
+	constexpr float PaintMaxDeviation = 0.15f;
+}
+
+void ULidarEditorTool_Base::Setup()
+{
+	Super::Setup();
+
+	ToolActions = CreateToolActions();
+
+	if(ToolActions)
+	{
+		AddToolPropertySource(ToolActions);
+	}
+}
+
+FText ULidarEditorTool_Base::GetToolMessage() const
+{
+	return FText();
+}
+
+void ULidarEditorTool_ClickDragBase::Setup()
+{
+	Super::Setup();
+
+	HoverBehavior = NewObject<UMouseHoverBehavior>();	
+	HoverBehavior->Modifiers.RegisterModifier(1, FInputDeviceState::IsShiftKeyDown);
+	HoverBehavior->Modifiers.RegisterModifier(2, FInputDeviceState::IsCtrlKeyDown);
+	HoverBehavior->Initialize(this);
+	AddInputBehavior(HoverBehavior, this);
+	
+	ClickDragBehavior = NewObject<UClickDragInputBehavior>();
+	ClickDragBehavior->Initialize(this);
+	ClickDragBehavior->Modifiers.RegisterModifier(1, FInputDeviceState::IsShiftKeyDown);
+	ClickDragBehavior->Modifiers.RegisterModifier(2, FInputDeviceState::IsCtrlKeyDown);
+	AddInputBehavior(ClickDragBehavior, this);
+}
+
+void ULidarEditorTool_ClickDragBase::Render(IToolsContextRenderAPI* RenderAPI)
+{
+	Super::Render(RenderAPI);
+	
+	GetToolManager()->GetContextQueriesAPI()->GetCurrentViewState(CameraState);
+}
+
+void ULidarEditorTool_ClickDragBase::OnUpdateModifierState(int ModifierID, bool bIsOn)
+{
+	switch (ModifierID)
+	{
+	case 1:
+		bShiftToggle = bIsOn;
+		break;
+	case 2:
+		bCtrlToggle = bIsOn;
+		break;
+	default:
+		break;
+	}
+}
+
+void ULidarToolActions_Align::AlignAroundWorldOrigin()
+{
+	FLidarPointCloudEditorHelper::AlignSelectionAroundWorldOrigin();
+}
+
+void ULidarToolActions_Align::AlignAroundOriginalCoordinates()
+{
+	FLidarPointCloudEditorHelper::SetOriginalCoordinateForSelection();
+}
+
+void ULidarToolActions_Align::ResetAlignment()
+{
+	FLidarPointCloudEditorHelper::CenterSelection();
+}
+
+void ULidarToolActions_Merge::MergeActors()
+{
+	FLidarPointCloudEditorHelper::MergeSelectionByComponent(bReplaceSourceActorsAfterMerging);
+}
+
+void ULidarToolActions_Merge::MergeData()
+{
+	FLidarPointCloudEditorHelper::MergeSelectionByData(bReplaceSourceActorsAfterMerging);
+}
+
+void ULidarToolActions_Collision::BuildCollision()
+ {
+ 	FLidarPointCloudEditorHelper::SetCollisionErrorForSelection(OverrideMaxCollisionError);
+ 	FLidarPointCloudEditorHelper::BuildCollisionForSelection();
+ }
+
+void ULidarToolActions_Collision::RemoveCollision()
+ {
+ 	FLidarPointCloudEditorHelper::RemoveCollisionForSelection();
+ }
+
+void ULidarToolActions_Normals::CalculateNormals()
+{
+	FLidarPointCloudEditorHelper::SetNormalsQualityForSelection(Quality, NoiseTolerance);
+	FLidarPointCloudEditorHelper::CalculateNormalsForSelection();
+}
+
+void ULidarToolActions_Selection::HideSelected()
+{
+	FLidarPointCloudEditorHelper::HideSelected();
+}
+
+void ULidarToolActions_Selection::ResetVisibility()
+{
+	FLidarPointCloudEditorHelper::ResetVisibility();
+}
+
+void ULidarToolActions_Selection::DeleteHidden()
+{
+	FLidarPointCloudEditorHelper::DeleteHidden();
+}
+
+void ULidarToolActions_Selection::Extract()
+{
+	FLidarPointCloudEditorHelper::Extract();
+}
+
+void ULidarToolActions_Selection::ExtractAsCopy()
+{
+	FLidarPointCloudEditorHelper::ExtractAsCopy();
+}
+
+void ULidarToolActions_Selection::CalculateNormals()
+{
+	FLidarPointCloudEditorHelper::CalculateNormals();
+	ClearSelection();
+}
+
+void ULidarToolActions_Selection::DeleteSelected()
+{
+	FLidarPointCloudEditorHelper::DeleteSelected();
+}
+
+void ULidarToolActions_Selection::InvertSelection()
+{
+	FLidarPointCloudEditorHelper::InvertSelection();
+}
+
+void ULidarToolActions_Selection::ClearSelection()
+{
+	FLidarPointCloudEditorHelper::ClearSelection();
+}
+
+void ULidarEditorTool_SelectionBase::DrawHUD(FCanvas* Canvas, IToolsContextRenderAPI* RenderAPI)
+{
+	if(Clicks.Num() == 0)
+	{
+		return;
+	}
+
+	const FLinearColor Color = GetHUDColor();
+	
+	for(int32 i = 1; i < Clicks.Num(); ++i)
+	{
+		FCanvasLineItem LineItem(Clicks[i - 1] / Canvas->GetDPIScale(), Clicks[i] / Canvas->GetDPIScale());
+		LineItem.SetColor(Color);
+		Canvas->DrawItem(LineItem);
+	}
+	
+	FCanvasLineItem LineItem(CurrentMousePos / Canvas->GetDPIScale(), Clicks.Last() / Canvas->GetDPIScale());
+	LineItem.SetColor(Color);
+	Canvas->DrawItem(LineItem);
+}
+
+TArray<FConvexVolume> ULidarEditorTool_SelectionBase::GetSelectionConvexVolumes()
+{
+	return FLidarPointCloudEditorHelper::BuildConvexVolumesFromPoints(Clicks);
+}
+
+TObjectPtr<UInteractiveToolPropertySet> ULidarEditorTool_SelectionBase::CreateToolActions()
+{
+	return NewObject<ULidarToolActions_Selection>(this);
+}
+
+bool ULidarEditorTool_SelectionBase::OnUpdateHover(const FInputDeviceRay& DevicePos)
+{
+	CurrentMousePos = DevicePos.ScreenPosition;
+	PostCurrentMousePosChanged();
+	return true;
+}
+
+void ULidarEditorTool_SelectionBase::OnClickPress(const FInputDeviceRay& PressPos)
+{
+	bSelecting = true;
+}
+
+void ULidarEditorTool_SelectionBase::OnClickRelease(const FInputDeviceRay& ReleasePos)
+{
+	bSelecting = false;
+}
+
+void ULidarEditorTool_SelectionBase::OnClickDrag(const FInputDeviceRay& DragPos)
+{
+	CurrentMousePos = DragPos.ScreenPosition;
+	PostCurrentMousePosChanged();
+}
+
+void ULidarEditorTool_SelectionBase::OnTerminateDragSequence()
+{
+	Clicks.Empty();
+	bSelecting = false;
+}
+
+bool ULidarEditorTool_SelectionBase::ExecuteNestedCancelCommand()
+{
+	OnTerminateDragSequence();
+	return true;
+}
+
+FText ULidarEditorTool_SelectionBase::GetToolMessage() const
+{
+	const FText ToolMessage = LOCTEXT("ULidarEditorToolToolMessage", "Use Left-click to start the selection. Hold Shift to add selection, hold Ctrl to subtract selection.");
+	return ToolMessage;
+}
+
+FLinearColor ULidarEditorTool_SelectionBase::GetHUDColor()
+{
+	return FLinearColor::White;
+}
+
+void ULidarEditorTool_SelectionBase::FinalizeSelection()
+{
+	const ELidarPointCloudSelectionMode SelectionMode = GetSelectionMode();
+
+	if(SelectionMode == ELidarPointCloudSelectionMode::None)
+	{
+		FLidarPointCloudEditorHelper::ClearSelection();
+	}
+	
+	const TArray<FConvexVolume> ConvexVolumes = GetSelectionConvexVolumes();
+	for (int32 i = 0; i < ConvexVolumes.Num(); ++i)
+	{
+		// Consecutive shapes need to be additive
+		FLidarPointCloudEditorHelper::SelectPointsByConvexVolume(ConvexVolumes[i], i > 0 ? ELidarPointCloudSelectionMode::Add : SelectionMode);
+	}
+}
+
+ELidarPointCloudSelectionMode ULidarEditorTool_SelectionBase::GetSelectionMode() const
+{
+	ELidarPointCloudSelectionMode SelectionMode = ELidarPointCloudSelectionMode::None;
+	if(bCtrlToggle)
+	{
+		SelectionMode = ELidarPointCloudSelectionMode::Subtract;
+	}
+	else if(bShiftToggle)
+	{
+		SelectionMode = ELidarPointCloudSelectionMode::Add;
+	}
+	return SelectionMode;
+}
+
+TArray<FConvexVolume> ULidarEditorTool_BoxSelection::GetSelectionConvexVolumes()
+{
+	return { FLidarPointCloudEditorHelper::BuildConvexVolumeFromCoordinates(Clicks[0], Clicks[2]) };
+}
+
+void ULidarEditorTool_BoxSelection::OnClickPress(const FInputDeviceRay& PressPos)
+{
+	Clicks.Append({
+		PressPos.ScreenPosition,
+		PressPos.ScreenPosition,
+		PressPos.ScreenPosition,
+		PressPos.ScreenPosition
+	});
+}
+
+void ULidarEditorTool_BoxSelection::OnClickDrag(const FInputDeviceRay& DragPos)
+{
+	if(Clicks.Num() == 0)
+	{
+		return;
+	}
+	
+	Clicks[1].Y = DragPos.ScreenPosition.Y;
+	Clicks[2] = DragPos.ScreenPosition;
+	Clicks[3].X = DragPos.ScreenPosition.X;
+}
+
+void ULidarEditorTool_BoxSelection::OnClickRelease(const FInputDeviceRay& ReleasePos)
+{
+	if(Clicks.Num() == 4)
+	{
+		if(Clicks[0] == Clicks[2])
+		{
+			FLidarPointCloudEditorHelper::ClearSelection();
+		}
+		else
+		{
+			FinalizeSelection();
+		}
+	}
+	
+	Clicks.Empty();
+}
+
+void ULidarEditorTool_PolygonalSelection::OnClickPress(const FInputDeviceRay& PressPos)
+{
+	if(IsWithinSnap())
+	{
+		FinalizeSelection();
+		Clicks.Empty();
+	}
+	else
+	{
+		Clicks.Add(PressPos.ScreenPosition);
+	}
+}
+
+FLinearColor ULidarEditorTool_PolygonalSelection::GetHUDColor()
+{
+	return IsWithinSnap() ? FLinearColor::Green : Super::GetHUDColor();
+}
+
+void ULidarEditorTool_PolygonalSelection::PostCurrentMousePosChanged()
+{
+	if(IsWithinSnap())
+	{
+		CurrentMousePos = Clicks[0];
+	}
+}
+
+bool ULidarEditorTool_PolygonalSelection::IsWithinSnap()
+{
+	return Clicks.Num() > 1 && (CurrentMousePos - Clicks[0]).SquaredLength() <= PolySnapDistanceSq;
+}
+
+void ULidarEditorTool_LassoSelection::OnClickPress(const FInputDeviceRay& PressPos)
+{
+	Clicks.Add(PressPos.ScreenPosition);
+}
+
+void ULidarEditorTool_LassoSelection::OnClickDrag(const FInputDeviceRay& DragPos)
+{
+	Super::OnClickDrag(DragPos);
+	
+	if(Clicks.Num() > 0 && (CurrentMousePos - Clicks.Last()).SquaredLength() >= LassoSpacingSq)
+	{
+		Clicks.Add(CurrentMousePos);
+	}
+}
+
+void ULidarEditorTool_LassoSelection::OnClickRelease(const FInputDeviceRay& ReleasePos)
+{
+	if(Clicks.Num() > 1)
+	{
+		FinalizeSelection();
+	}
+
+	Clicks.Empty();
+}
+
+void ULidarEditorTool_PaintSelection::Setup()
+{
+	Super::Setup();
+	BrushRadius = GetDefault<ULidarToolActions_PaintSelection>()->BrushRadius;
+}
+
+void ULidarEditorTool_PaintSelection::Render(IToolsContextRenderAPI* RenderAPI)
+{
+	Super::Render(RenderAPI);
+	
+	if(bHasHit)
+	{
+		DrawWireSphere(RenderAPI->GetPrimitiveDrawInterface(), (FVector)HitLocation, FLinearColor::Red, BrushRadius, 32, SDPG_World);
+	}
+}
+
+void ULidarEditorTool_PaintSelection::PostCurrentMousePosChanged()
+{
+	const FLidarPointCloudRay Ray = FLidarPointCloudEditorHelper::MakeRayFromScreenPosition(CurrentMousePos);
+
+	FVector3f NewHitLocation;
+	bHasHit = FLidarPointCloudEditorHelper::RayTracePointClouds(Ray, 1, NewHitLocation);
+
+	if(!bHasHit)
+	{
+		return;
+	}
+	
+	const float NewDistance = FVector3f::Dist(NewHitLocation, Ray.Origin);
+	const float Deviation = (NewDistance - LastHitDistance) / LastHitDistance;
+
+	// If painting, prevent large depth changes
+	// If not, query larger trace radius - if it passes the deviation test, it was a gap
+	if (Deviation > PaintMaxDeviation && (bSelecting || (
+						FLidarPointCloudEditorHelper::RayTracePointClouds(Ray, 6, NewHitLocation) &&
+						(FVector3f::Dist(NewHitLocation, Ray.Origin) - LastHitDistance) / LastHitDistance <= PaintMaxDeviation)))
+	{
+		HitLocation = FVector3f(Ray.Origin + Ray.GetDirection() * LastHitDistance);
+		return;
+	}
+	
+	HitLocation = NewHitLocation;
+	LastHitDistance = NewDistance;
+}
+
+void ULidarEditorTool_PaintSelection::OnClickPress(const FInputDeviceRay& PressPos)
+{
+	Super::OnClickPress(PressPos);
+
+	if(GetSelectionMode() == ELidarPointCloudSelectionMode::None)
+	{
+		FLidarPointCloudEditorHelper::ClearSelection();
+	}
+	
+	Paint();
+}
+
+void ULidarEditorTool_PaintSelection::OnClickDrag(const FInputDeviceRay& DragPos)
+{
+	Super::OnClickDrag(DragPos);
+	Paint();
+}
+
+void ULidarEditorTool_PaintSelection::OnPropertyModified(UObject* PropertySet, FProperty* Property)
+{
+	if(const ULidarToolActions_PaintSelection* Actions = Cast<ULidarToolActions_PaintSelection>(PropertySet))
+	{
+		if(Property && Property->GetName().Equals("BrushRadius"))
+		{
+			BrushRadius = Actions->BrushRadius;
+		}
+	}
+}
+
+void ULidarEditorTool_PaintSelection::Paint()
+{
+	if(bHasHit)
+	{
+		FLidarPointCloudEditorHelper::SelectPointsBySphere(FSphere((FVector)HitLocation, BrushRadius), GetSelectionMode());
+	}
+}
+
+#undef LOCTEXT_NAMESPACE
diff --git a/Engine/Plugins/Enterprise/LidarPointCloud/Source/LidarPointCloudEditor/Private/LidarPointCloudEditorTools.h b/Engine/Plugins/Enterprise/LidarPointCloud/Source/LidarPointCloudEditor/Private/LidarPointCloudEditorTools.h
new file mode 100644
--- /dev/null
+++ b/Engine/Plugins/Enterprise/LidarPointCloud/Source/LidarPointCloudEditor/Private/LidarPointCloudEditorTools.h
@@ -0,0 +1,385 @@
+// Copyright Epic Games, Inc. All Rights Reserved.
+
+#pragma once
+#include "InteractiveTool.h"
+#include "InteractiveToolBuilder.h"
+#include "InteractiveToolManager.h"
+#include "InteractiveToolQueryInterfaces.h"
+#include "LidarPointCloudEditorHelper.h"
+#include "BaseBehaviors/BehaviorTargetInterfaces.h"
+#include "LidarPointCloudEditorTools.generated.h"
+
+class UMouseHoverBehavior;
+class URectangleMarqueeMechanic;
+class UClickDragInputBehavior;
+
+UCLASS()
+class ULidarEditorTool_Base : public UInteractiveTool
+{
+public:
+	GENERATED_BODY()
+	virtual void Setup() override;
+	virtual TObjectPtr<UInteractiveToolPropertySet> CreateToolActions() { return nullptr; }
+	virtual FText GetToolMessage() const;
+
+	UPROPERTY()
+	TObjectPtr<UInteractiveToolPropertySet> ToolActions = nullptr;
+};
+
+UCLASS()
+class ULidarEditorToolBuilder_Base : public UInteractiveToolWithToolTargetsBuilder
+{
+	GENERATED_BODY()
+
+public:
+	virtual bool CanBuildTool(const FToolBuilderState& SceneState) const override { return true; }
+	virtual UInteractiveTool* BuildTool(const FToolBuilderState& SceneState) const override { return NewObject<ULidarEditorTool_Base>(SceneState.ToolManager); }
+};
+
+UCLASS()
+class ULidarEditorTool_ClickDragBase :
+	public ULidarEditorTool_Base,
+	public IClickDragBehaviorTarget,
+	public IHoverBehaviorTarget,
+	public IInteractiveToolNestedAcceptCancelAPI
+{
+public:
+	GENERATED_BODY()
+	
+	virtual void Setup() override;
+	virtual void Render(IToolsContextRenderAPI* RenderAPI) override;
+
+	// IClickDragBehaviorTarget implementation
+	virtual FInputRayHit CanBeginClickDragSequence(const FInputDeviceRay& PressPos) override { return FInputRayHit(TNumericLimits<float>::Max()); }
+	virtual void OnClickPress(const FInputDeviceRay& PressPos) override {}
+	virtual void OnClickDrag(const FInputDeviceRay& DragPos) override {}
+	virtual void OnClickRelease(const FInputDeviceRay& ReleasePos) override {}
+	virtual void OnTerminateDragSequence() override {}
+
+	// IHoverBehaviorTarget implementation
+	virtual FInputRayHit BeginHoverSequenceHitTest(const FInputDeviceRay& PressPos) override { return FInputRayHit(TNumericLimits<float>::Max()); }
+	virtual void OnBeginHover(const FInputDeviceRay& DevicePos) override {}
+	virtual bool OnUpdateHover(const FInputDeviceRay& DevicePos) override { return true; }
+	virtual void OnEndHover() override {}
+	virtual void OnUpdateModifierState(int ModifierID, bool bIsOn) override;
+
+	// IInteractiveToolNestedAcceptCancelAPI implementation
+	virtual bool SupportsNestedCancelCommand() override { return true; }
+	virtual bool CanCurrentlyNestedCancel() override { return true; }
+	virtual bool ExecuteNestedCancelCommand() override { return false; }
+
+	UPROPERTY()
+	TObjectPtr<UClickDragInputBehavior> ClickDragBehavior = nullptr;
+
+	UPROPERTY()
+	TObjectPtr<UMouseHoverBehavior> HoverBehavior = nullptr;
+
+protected:
+	FViewCameraState CameraState;
+	
+	bool bShiftToggle = false;
+	bool bCtrlToggle = false;
+};
+
+// -------------------------------------------------------
+
+UCLASS()
+class ULidarEditorToolBuilder_Select : public ULidarEditorToolBuilder_Base
+{
+	GENERATED_BODY()
+};
+
+UCLASS()
+class ULidarToolActions_Align : public UInteractiveToolPropertySet
+{
+	GENERATED_BODY()
+
+public:
+	UFUNCTION(CallInEditor, Category = Actions)
+	void AlignAroundWorldOrigin();
+
+	UFUNCTION(CallInEditor, Category = Actions)
+	void AlignAroundOriginalCoordinates();
+
+	UFUNCTION(CallInEditor, Category = Actions)
+	void ResetAlignment();
+};
+
+UCLASS()
+class ULidarEditorTool_Align : public ULidarEditorTool_Base
+{
+public:
+	GENERATED_BODY()
+	virtual TObjectPtr<UInteractiveToolPropertySet> CreateToolActions() override { return NewObject<ULidarToolActions_Align>(this); }
+};
+
+UCLASS()
+class ULidarEditorToolBuilder_Align : public ULidarEditorToolBuilder_Base
+{
+public:
+	GENERATED_BODY()
+	virtual UInteractiveTool* BuildTool(const FToolBuilderState& SceneState) const override { return NewObject<ULidarEditorTool_Align>(SceneState.ToolManager); }
+};
+
+UCLASS()
+class ULidarToolActions_Merge : public UInteractiveToolPropertySet
+{
+	GENERATED_BODY()
+
+public:
+	UPROPERTY(EditAnywhere, Category = Options)
+	bool bReplaceSourceActorsAfterMerging = false;
+	
+	UFUNCTION(CallInEditor, Category = Actions)
+	void MergeActors();
+
+	UFUNCTION(CallInEditor, Category = Actions)
+	void MergeData();
+};
+
+UCLASS()
+class ULidarEditorTool_Merge : public ULidarEditorTool_Base
+{
+public:
+	GENERATED_BODY()
+	virtual TObjectPtr<UInteractiveToolPropertySet> CreateToolActions() override { return NewObject<ULidarToolActions_Merge>(this); }
+};
+
+UCLASS()
+class ULidarEditorToolBuilder_Merge : public ULidarEditorToolBuilder_Base
+{
+public:
+	GENERATED_BODY()
+	virtual UInteractiveTool* BuildTool(const FToolBuilderState& SceneState) const override { return NewObject<ULidarEditorTool_Merge>(SceneState.ToolManager); }
+};
+
+UCLASS()
+class ULidarToolActions_Collision : public UInteractiveToolPropertySet
+{
+	GENERATED_BODY()
+
+public:
+	UPROPERTY(EditAnywhere, Category = Options, meta = (UIMin = "0", UIMax = "2000", ClampMin = "0"))
+	float OverrideMaxCollisionError = 0;
+	
+	UFUNCTION(CallInEditor, Category = Actions, meta = (DisplayName = "(Re-)Build Collision"))
+	void BuildCollision();
+
+	UFUNCTION(CallInEditor, Category = Actions)
+	void RemoveCollision();
+};
+
+UCLASS()
+class ULidarEditorTool_Collision : public ULidarEditorTool_Base
+{
+public:
+	GENERATED_BODY()
+	virtual TObjectPtr<UInteractiveToolPropertySet> CreateToolActions() override { return NewObject<ULidarToolActions_Collision>(this); }
+};
+
+UCLASS()
+class ULidarEditorToolBuilder_Collision : public ULidarEditorToolBuilder_Base
+{
+public:
+	GENERATED_BODY()
+	virtual UInteractiveTool* BuildTool(const FToolBuilderState& SceneState) const override { return NewObject<ULidarEditorTool_Collision>(SceneState.ToolManager); }
+};
+
+UCLASS()
+class ULidarToolActions_Normals : public UInteractiveToolPropertySet
+{
+	GENERATED_BODY()
+
+public:
+	/** Higher values will generally result in more accurate calculations, at the expense of time */
+	UPROPERTY(EditAnywhere, Category = "Options", meta = (ClampMin = "1", ClampMax = "100"))
+	int32 Quality = 40;
+
+	/**
+	 * Higher values are less susceptible to noise, but will most likely lose finer details, especially around hard edges.
+	 * Lower values retain more detail, at the expense of time.
+	 * NOTE: setting this too low will cause visual artifacts and geometry holes in noisier datasets.
+	 */
+	UPROPERTY(EditAnywhere, BlueprintReadWrite, Category = "Options", meta = (ClampMin = "0.0"))
+	float NoiseTolerance = 1;
+	
+	UFUNCTION(CallInEditor, Category = Actions, meta = (DisplayName = "(Re-)Build Normals"))
+	void CalculateNormals();
+};
+
+UCLASS()
+class ULidarEditorTool_Normals : public ULidarEditorTool_Base
+{
+public:
+	GENERATED_BODY()
+	virtual TObjectPtr<UInteractiveToolPropertySet> CreateToolActions() override { return NewObject<ULidarToolActions_Normals>(this); }
+};
+
+UCLASS()
+class ULidarEditorToolBuilder_Normals : public ULidarEditorToolBuilder_Base
+{
+public:
+	GENERATED_BODY()
+	virtual UInteractiveTool* BuildTool(const FToolBuilderState& SceneState) const override { return NewObject<ULidarEditorTool_Normals>(SceneState.ToolManager); }
+};
+
+UCLASS()
+class ULidarToolActions_Selection : public UInteractiveToolPropertySet
+{
+	GENERATED_BODY()
+
+public:
+	UFUNCTION(CallInEditor, Category = Visibility)
+	void HideSelected();
+
+	UFUNCTION(CallInEditor, Category = Visibility)
+	void ResetVisibility();
+
+	UFUNCTION(CallInEditor, Category = Selection)
+	void InvertSelection();
+
+	UFUNCTION(CallInEditor, Category = Selection)
+	void ClearSelection();
+
+	UFUNCTION(CallInEditor, Category = Deletion)
+	void DeleteSelected();
+	
+	UFUNCTION(CallInEditor, Category = Deletion)
+	void DeleteHidden();
+
+	UFUNCTION(CallInEditor, Category = Extraction)
+	void Extract();
+
+	UFUNCTION(CallInEditor, Category = Extraction)
+	void ExtractAsCopy();
+
+	UFUNCTION(CallInEditor, Category = Normals)
+	void CalculateNormals();
+};
+
+UCLASS()
+class ULidarEditorTool_SelectionBase : public ULidarEditorTool_ClickDragBase
+{
+public:
+	GENERATED_BODY()
+	virtual void DrawHUD(FCanvas* Canvas, IToolsContextRenderAPI* RenderAPI) override;
+	virtual TArray<FConvexVolume> GetSelectionConvexVolumes();
+	virtual TObjectPtr<UInteractiveToolPropertySet> CreateToolActions() override;
+	virtual bool OnUpdateHover(const FInputDeviceRay& DevicePos) override;
+	virtual void OnClickPress(const FInputDeviceRay& PressPos) override;
+	virtual void OnClickDrag(const FInputDeviceRay& DragPos) override;
+	virtual void OnClickRelease(const FInputDeviceRay& ReleasePos) override;
+	virtual void OnTerminateDragSequence() override;
+	virtual bool ExecuteNestedCancelCommand() override;
+	virtual FText GetToolMessage() const override;
+	virtual FLinearColor GetHUDColor();
+	virtual void FinalizeSelection();
+
+	virtual void PostCurrentMousePosChanged() {}
+	virtual ELidarPointCloudSelectionMode GetSelectionMode() const;
+
+protected:
+	FVector2d CurrentMousePos;
+	TArray<FVector2d> Clicks;
+	bool bSelecting;
+};
+
+UCLASS()
+class ULidarEditorTool_BoxSelection : public ULidarEditorTool_SelectionBase
+{
+public:
+	GENERATED_BODY()
+	virtual TArray<FConvexVolume> GetSelectionConvexVolumes() override;
+	virtual void OnClickPress(const FInputDeviceRay& PressPos) override;
+	virtual void OnClickDrag(const FInputDeviceRay& DragPos) override;
+	virtual void OnClickRelease(const FInputDeviceRay& ReleasePos) override;
+};
+
+UCLASS()
+class ULidarEditorToolBuilder_BoxSelection : public ULidarEditorToolBuilder_Base
+{
+public:
+	GENERATED_BODY()
+	virtual UInteractiveTool* BuildTool(const FToolBuilderState& SceneState) const override { return NewObject<ULidarEditorTool_BoxSelection>(SceneState.ToolManager); }
+};
+
+UCLASS()
+class ULidarEditorTool_PolygonalSelection : public ULidarEditorTool_SelectionBase
+{
+public:
+	GENERATED_BODY()
+	virtual void OnClickPress(const FInputDeviceRay& PressPos) override;
+	virtual FLinearColor GetHUDColor() override;
+	virtual void PostCurrentMousePosChanged() override;
+
+private:
+	bool IsWithinSnap();
+};
+
+UCLASS()
+class ULidarEditorToolBuilder_PolygonalSelection : public ULidarEditorToolBuilder_Base
+{
+public:
+	GENERATED_BODY()
+	virtual UInteractiveTool* BuildTool(const FToolBuilderState& SceneState) const override { return NewObject<ULidarEditorTool_PolygonalSelection>(SceneState.ToolManager); }
+};
+
+UCLASS()
+class ULidarEditorTool_LassoSelection : public ULidarEditorTool_SelectionBase
+{
+public:
+	GENERATED_BODY()
+	virtual void OnClickPress(const FInputDeviceRay& PressPos) override;
+	virtual void OnClickDrag(const FInputDeviceRay& DragPos) override;
+	virtual void OnClickRelease(const FInputDeviceRay& ReleasePos) override;
+};
+
+UCLASS()
+class ULidarEditorToolBuilder_LassoSelection : public ULidarEditorToolBuilder_Base
+{
+public:
+	GENERATED_BODY()
+	virtual UInteractiveTool* BuildTool(const FToolBuilderState& SceneState) const override { return NewObject<ULidarEditorTool_LassoSelection>(SceneState.ToolManager); }
+};
+
+UCLASS()
+class ULidarToolActions_PaintSelection : public ULidarToolActions_Selection
+{
+	GENERATED_BODY()
+
+public:
+	UPROPERTY(EditAnywhere, Category = Brush, meta = (UIMin = "0", UIMax = "8196", ClampMin = "0", DisplayPriority = 1))
+	float BrushRadius = 250;
+};
+
+UCLASS()
+class ULidarEditorTool_PaintSelection : public ULidarEditorTool_SelectionBase
+{
+public:
+	GENERATED_BODY()
+	virtual void Setup() override;
+	virtual void DrawHUD(FCanvas* Canvas, IToolsContextRenderAPI* RenderAPI) override {}
+	virtual void Render(IToolsContextRenderAPI* RenderAPI) override;
+	virtual void PostCurrentMousePosChanged() override;
+	virtual void OnClickPress(const FInputDeviceRay& PressPos) override;
+	virtual void OnClickDrag(const FInputDeviceRay& DragPos) override;
+	virtual TObjectPtr<UInteractiveToolPropertySet> CreateToolActions() override { return NewObject<ULidarToolActions_PaintSelection>(this); }
+	virtual void OnPropertyModified(UObject* PropertySet, FProperty* Property) override;
+
+private:
+	void Paint();
+
+private:
+	FVector3f HitLocation;
+	float LastHitDistance;
+	bool bHasHit;
+	float BrushRadius;
+};
+
+UCLASS()
+class ULidarEditorToolBuilder_PaintSelection : public ULidarEditorToolBuilder_Base
+{
+public:
+	GENERATED_BODY()
+	virtual UInteractiveTool* BuildTool(const FToolBuilderState& SceneState) const override { return NewObject<ULidarEditorTool_PaintSelection>(SceneState.ToolManager); }
+};
\ No newline at end of file
diff --git a/Engine/Plugins/Enterprise/LidarPointCloud/Source/LidarPointCloudEditor/Private/LidarPointCloudEditorViewportClient.cpp b/Engine/Plugins/Enterprise/LidarPointCloud/Source/LidarPointCloudEditor/Private/LidarPointCloudEditorViewportClient.cpp
--- a/Engine/Plugins/Enterprise/LidarPointCloud/Source/LidarPointCloudEditor/Private/LidarPointCloudEditorViewportClient.cpp
+++ b/Engine/Plugins/Enterprise/LidarPointCloud/Source/LidarPointCloudEditor/Private/LidarPointCloudEditorViewportClient.cpp
@@ -52,47 +52,6 @@ TArray<FVector2D> ToVectorArray(const TArray<FIntPoint>& Points)
 	return VectorPoints;
 }
 
-// Slow, O(n2), but sufficient for the current problem
-bool IsPolygonSelfIntersecting(const TArray<FVector2D>& Points, bool bAllowLooping)
-{
-	const int32 MaxIndex = bAllowLooping ? Points.Num() : Points.Num() - 1;
-
-	for (int32 i = 0; i < MaxIndex; ++i)
-	{
-		const int32 i1 = i < Points.Num() - 1 ? i + 1 : 0;
-
-		const FVector2D P1 = Points[i];
-		const FVector2D P2 = Points[i1];
-
-		for (int32 j = 0; j < MaxIndex; ++j)
-		{
-			const int32 j1 = j < Points.Num() - 1 ? j + 1 : 0;
-
-			if (j1 != i && j != i && j != i1)
-			{
-				// Modified, inlined FMath::SegmentIntersection2D
-				const FVector2D SegmentStartA = P1;
-				const FVector2D SegmentEndA = P2;
-				const FVector2D SegmentStartB = Points[j];
-				const FVector2D SegmentEndB = Points[j1];
-				const FVector2D VectorA = P2 - SegmentStartA;
-				const FVector2D VectorB = SegmentEndB - SegmentStartB;
-
-				const float S = (-VectorA.Y * (SegmentStartA.X - SegmentStartB.X) + VectorA.X * (SegmentStartA.Y - SegmentStartB.Y)) / (-VectorB.X * VectorA.Y + VectorA.X * VectorB.Y);
-				const float T = (VectorB.X * (SegmentStartA.Y - SegmentStartB.Y) - VectorB.Y * (SegmentStartA.X - SegmentStartB.X)) / (-VectorB.X * VectorA.Y + VectorA.X * VectorB.Y);
-
-				if (S >= 0 && S <= 1 && T >= 0 && T <= 1)
-				{
-					return true;
-				}
-			}
-		}
-	}
-
-	return false;
-}
-
-
 // Copied from GeomTools.cpp and converted to work with FIntPoint
 bool IsPolygonConvex(const TArray<FIntPoint>& Points)
 {
@@ -292,7 +251,7 @@ bool FLidarPointCloudEditorViewportClient::InputKey(FViewport* InViewport, int32
 							TArray<FVector2D> VectorPoints = ToVectorArray(SelectionPoints);
 							VectorPoints.Add(NewPoint);
 
-							if (!IsPolygonSelfIntersecting(VectorPoints, false))
+							if (!FLidarPointCloudEditorHelper::IsPolygonSelfIntersecting(VectorPoints, false))
 							{
 								// Snap to first point
 								if (SelectionPoints.Num() > 1 && (NewPoint - SelectionPoints[0]).SizeSquared() < PolySnapDistanceSq)
@@ -690,7 +649,7 @@ void FLidarPointCloudEditorViewportClient::OnPolygonalSelectionEnd()
 			else
 			{
 				// Check for self-intersecting shape
-				if (!IsPolygonSelfIntersecting(VectorPoints, true))
+				if (!FLidarPointCloudEditorHelper::IsPolygonSelfIntersecting(VectorPoints, true))
 				{
 					// The separation needs points in CCW order
 					if (!FGeomTools2D::IsPolygonWindingCCW(VectorPoints))
@@ -821,7 +780,7 @@ void FLidarPointCloudEditorViewportClient::DrawSelectionPolygonal(FCanvas& Canva
 
 	// Calculate visual indication of complete polygon for the user
 	const bool bPolyComplete = DrawSelectionPoints.Num() > 2 && (DrawSelectionPoints.Last() - DrawSelectionPoints[0]).SizeSquared() < PolySnapDistanceSq;
-	const bool bSelfIntersecting = DrawSelectionPoints.Num() > 2 && IsPolygonSelfIntersecting(VectorPoints, true);
+	const bool bSelfIntersecting = DrawSelectionPoints.Num() > 2 && FLidarPointCloudEditorHelper::IsPolygonSelfIntersecting(VectorPoints, true);
 	FLinearColor SelectionColor = bSelfIntersecting ? FLinearColor::Red : bPolyComplete ? FLinearColor::Green : GetDefault<UEditorStyleSettings>()->SelectionColor;
 
 	// Selection Area
diff --git a/Engine/Plugins/Enterprise/LidarPointCloud/Source/LidarPointCloudEditor/Private/LidarPointCloudEditorViewportClient.h b/Engine/Plugins/Enterprise/LidarPointCloud/Source/LidarPointCloudEditor/Private/LidarPointCloudEditorViewportClient.h
--- a/Engine/Plugins/Enterprise/LidarPointCloud/Source/LidarPointCloudEditor/Private/LidarPointCloudEditorViewportClient.h
+++ b/Engine/Plugins/Enterprise/LidarPointCloud/Source/LidarPointCloudEditor/Private/LidarPointCloudEditorViewportClient.h
@@ -6,7 +6,7 @@
 #include "InputCoreTypes.h"
 #include "UnrealWidgetFwd.h"
 #include "EditorViewportClient.h"
-#include "Components.h"
+#include "LidarPointCloudEditorHelper.h"
 #include "LidarPointCloudShared.h"
 
 class FAdvancedPreviewScene;
@@ -16,13 +16,6 @@ class SLidarPointCloudEditorViewport;
 class ULidarPointCloud;
 class ULidarPointCloudComponent;
 
-enum class ELidarPointCloudSelectionMode : uint8
-{
-	None,
-	Add,
-	Subtract
-};
-
 enum class ELidarPointCloudSelectionMethod : uint8
 {
 	Box,
diff --git a/Engine/Plugins/Enterprise/LidarPointCloud/Source/LidarPointCloudEditor/Private/LidarPointCloudFactory.cpp b/Engine/Plugins/Enterprise/LidarPointCloud/Source/LidarPointCloudEditor/Private/LidarPointCloudFactory.cpp
--- a/Engine/Plugins/Enterprise/LidarPointCloud/Source/LidarPointCloudEditor/Private/LidarPointCloudFactory.cpp
+++ b/Engine/Plugins/Enterprise/LidarPointCloud/Source/LidarPointCloudEditor/Private/LidarPointCloudFactory.cpp
@@ -9,6 +9,8 @@
 #include "Styling/AppStyle.h"
 #include "Framework/MultiBox/MultiBoxBuilder.h"
 #include "Editor.h"
+#include "LidarPointCloudEditorHelper.h"
+
 #include "Misc/ScopedSlowTask.h"
 #include "AssetRegistry/AssetRegistryModule.h"
 #include "PackageTools.h"
@@ -132,17 +134,8 @@ void FAssetTypeActions_LidarPointCloud::ExecuteMerge(TArray<ULidarPointCloud*> P
 		MergedCloudPackage->SetPackageFlags(PKG_NewlyCreated);
 
 		ULidarPointCloud* PC = NewObject<ULidarPointCloud>(MergedCloudPackage, FName(*FPaths::GetBaseFilename(MergedCloudPackage->GetName())), EObjectFlags::RF_Public | EObjectFlags::RF_Standalone | EObjectFlags::RF_Transactional);
-		if (IsValid(PC))
-		{
-			FScopedSlowTask ProgressDialog(PointClouds.Num() + 2, LOCTEXT("Merge", "Merging Point Clouds..."));
-			ProgressDialog.MakeDialog();
-
-			PC->Merge(PointClouds, [&ProgressDialog]() { ProgressDialog.EnterProgressFrame(1.f); });
 
-			PC->MarkPackageDirty();
-
-			FAssetRegistryModule::AssetCreated(PC);
-		}
+		FLidarPointCloudEditorHelper::MergeLidar(PC, PointClouds);
 	}
 }
 
diff --git a/Engine/Plugins/Enterprise/LidarPointCloud/Source/LidarPointCloudEditor/Private/LidarPointCloudStyle.cpp b/Engine/Plugins/Enterprise/LidarPointCloud/Source/LidarPointCloudEditor/Private/LidarPointCloudStyle.cpp
--- a/Engine/Plugins/Enterprise/LidarPointCloud/Source/LidarPointCloudEditor/Private/LidarPointCloudStyle.cpp
+++ b/Engine/Plugins/Enterprise/LidarPointCloud/Source/LidarPointCloudEditor/Private/LidarPointCloudStyle.cpp
@@ -7,7 +7,7 @@
 #include "Styling/AppStyle.h"
 #include "Interfaces/IPluginManager.h"
 
-#define IMAGE_PLUGIN_BRUSH( RelativePath, ... ) FSlateImageBrush( FLidarPointCloudStyle::InContent( RelativePath, ".png" ), __VA_ARGS__ )
+#define IMAGE_PLUGIN_BRUSH( RelativePath, ... ) FSlateImageBrush( StyleSet->RootToContentDir( RelativePath, TEXT(".png") ), __VA_ARGS__ )
 
 TSharedPtr<FSlateStyleSet> FLidarPointCloudStyle::StyleSet = nullptr;
 TSharedPtr<class ISlateStyle> FLidarPointCloudStyle::Get() { return StyleSet; }
@@ -16,6 +16,7 @@ void FLidarPointCloudStyle::Initialize()
 {
 	// Const icon & thumbnail sizes
 	const FVector2D Icon16x16(16.0f, 16.0f);
+	const FVector2D Icon20x20(20.0f, 20.0f);
 	const FVector2D Icon32x32(32.0f, 32.0f);
 	const FVector2D Icon40x40(40.0f, 40.0f);
 	const FVector2D Icon128x128(128.0f, 128.0f);
@@ -25,15 +26,11 @@ void FLidarPointCloudStyle::Initialize()
 	{
 		return;
 	}
-
-
+	
 	StyleSet = MakeShareable(new FSlateStyleSet(GetStyleSetName()));
-	StyleSet->SetContentRoot(FPaths::EngineContentDir() / TEXT("Editor/Slate"));
-	StyleSet->SetCoreContentRoot(FPaths::EngineContentDir() / TEXT("Slate"));
-
-	FString ContentDir = IPluginManager::Get().FindPlugin(TEXT("LidarPointCloud"))->GetContentDir();
-	StyleSet->SetContentRoot(ContentDir);
-
+	
+	StyleSet->SetContentRoot(IPluginManager::Get().FindPlugin(TEXT("LidarPointCloud"))->GetContentDir() / TEXT("Icons"));
+	
 	StyleSet->Set("ClassIcon.LidarPointCloud", new IMAGE_PLUGIN_BRUSH("icon_32", Icon16x16));
 	StyleSet->Set("ClassIcon32.LidarPointCloud", new IMAGE_PLUGIN_BRUSH("icon_32", Icon32x32));
 	StyleSet->Set("ClassThumbnail.LidarPointCloud", new IMAGE_PLUGIN_BRUSH("icon_128", Icon128x128));
@@ -69,6 +66,19 @@ void FLidarPointCloudStyle::Initialize()
 	StyleSet->Set("LidarPointCloudEditor.ExtractCopy", new IMAGE_PLUGIN_BRUSH("icon_extractcopy_40", Icon40x40));
 	StyleSet->Set("LidarPointCloudEditor.Merge", new IMAGE_PLUGIN_BRUSH("icon_merge_40", Icon40x40));
 	StyleSet->Set("LidarPointCloudEditor.Align", new IMAGE_PLUGIN_BRUSH("icon_align_40", Icon40x40));
+	
+	StyleSet->Set("LidarPointCloudEditor.ToolkitCollision", new IMAGE_PLUGIN_BRUSH("icon_collision_40", Icon20x20));
+	StyleSet->Set("LidarPointCloudEditor.ToolkitMerge", new IMAGE_PLUGIN_BRUSH("icon_merge_40", Icon20x20));
+	StyleSet->Set("LidarPointCloudEditor.ToolkitAlign", new IMAGE_PLUGIN_BRUSH("icon_align_40", Icon20x20));
+	StyleSet->Set("LidarPointCloudEditor.ToolkitBoxSelection", new IMAGE_PLUGIN_BRUSH("icon_selbox_40", Icon20x20));
+	StyleSet->Set("LidarPointCloudEditor.ToolkitPolygonalSelection", new IMAGE_PLUGIN_BRUSH("icon_selpoly_40", Icon20x20));
+	StyleSet->Set("LidarPointCloudEditor.ToolkitLassoSelection", new IMAGE_PLUGIN_BRUSH("icon_sellasso_40", Icon20x20));
+	StyleSet->Set("LidarPointCloudEditor.ToolkitPaintSelection", new IMAGE_PLUGIN_BRUSH("icon_selpaint_40", Icon20x20));
+	
+	StyleSet->SetContentRoot(FPaths::EngineContentDir() / TEXT("Editor/Slate"));
+	
+	StyleSet->Set("LidarPointCloudEditor.ToolkitSelect", new IMAGE_PLUGIN_BRUSH( TEXT("Icons/GeneralTools/Select_40x"), Icon20x20 ) );
+	StyleSet->Set("LidarPointCloudEditor.ToolkitNormals", new IMAGE_PLUGIN_BRUSH( TEXT("Icons/icon_StaticMeshEd_Normals_40x"), Icon20x20 ) );
 
 	FSlateStyleRegistry::RegisterSlateStyle(*StyleSet.Get());
 }
@@ -91,9 +101,3 @@ FName FLidarPointCloudStyle::GetStyleSetName()
 	return PaperStyleName;
 }
 
-FString FLidarPointCloudStyle::InContent(const FString& RelativePath, const ANSICHAR* Extension)
-{
-	static FString ContentDir = IPluginManager::Get().FindPlugin(TEXT("LidarPointCloud"))->GetContentDir() / TEXT("Icons");
-	return (ContentDir / RelativePath) + Extension;
-}
-
diff --git a/Engine/Plugins/Enterprise/LidarPointCloud/Source/LidarPointCloudEditor/Private/LidarPointCloudStyle.h b/Engine/Plugins/Enterprise/LidarPointCloud/Source/LidarPointCloudEditor/Private/LidarPointCloudStyle.h
--- a/Engine/Plugins/Enterprise/LidarPointCloud/Source/LidarPointCloudEditor/Private/LidarPointCloudStyle.h
+++ b/Engine/Plugins/Enterprise/LidarPointCloud/Source/LidarPointCloudEditor/Private/LidarPointCloudStyle.h
@@ -16,9 +16,7 @@ public:
 	static TSharedPtr< class ISlateStyle > Get();
 
 	static FName GetStyleSetName();
-private:
-	static FString InContent(const FString& RelativePath, const ANSICHAR* Extension);
 
 private:
-	static TSharedPtr< class FSlateStyleSet > StyleSet;
+	static TSharedPtr<class FSlateStyleSet> StyleSet;
 };
diff --git a/Engine/Plugins/Enterprise/LidarPointCloud/Source/LidarPointCloudRuntime/Private/LidarPointCloud.cpp b/Engine/Plugins/Enterprise/LidarPointCloud/Source/LidarPointCloudRuntime/Private/LidarPointCloud.cpp
--- a/Engine/Plugins/Enterprise/LidarPointCloud/Source/LidarPointCloudRuntime/Private/LidarPointCloud.cpp
+++ b/Engine/Plugins/Enterprise/LidarPointCloud/Source/LidarPointCloudRuntime/Private/LidarPointCloud.cpp
@@ -545,6 +545,69 @@ bool ULidarPointCloud::LineTraceSingle(FVector Origin, FVector Direction, float
 	return false;
 }
 
+#if WITH_EDITOR
+void ULidarPointCloud::SelectByConvexVolume(FConvexVolume ConvexVolume, bool bAdditive, bool bApplyLocationOffset, bool bVisibleOnly)
+{
+	if(bApplyLocationOffset)
+	{
+		for(FPlane& Plane : ConvexVolume.Planes)
+		{
+			Plane = Plane.TranslateBy(-LocationOffset);
+		}
+
+		ConvexVolume.Init();
+	}
+
+	Octree.SelectByConvexVolume(ConvexVolume, bAdditive, bVisibleOnly);
+}
+
+void ULidarPointCloud::SelectBySphere(FSphere Sphere, bool bAdditive, bool bApplyLocationOffset, bool bVisibleOnly)
+{
+	if(bApplyLocationOffset)
+	{
+		Sphere.Center -= LocationOffset;
+	}
+
+	Octree.SelectBySphere(Sphere, bAdditive, bVisibleOnly);
+}
+
+void ULidarPointCloud::HideSelected()
+{
+	Octree.HideSelected();
+}
+
+void ULidarPointCloud::DeleteSelected()
+{
+	Octree.DeleteSelected();
+}
+
+void ULidarPointCloud::InvertSelection()
+{
+	Octree.InvertSelection();
+}
+
+int64 ULidarPointCloud::NumSelectedPoints()
+{
+	return Octree.NumSelectedPoints();
+}
+
+void ULidarPointCloud::GetSelectedPointsAsCopies(TArray64<FLidarPointCloudPoint>& SelectedPoints, FTransform Transform)
+{
+	Transform.AddToTranslation(LocationOffset);
+	Octree.GetSelectedPointsAsCopies(SelectedPoints, Transform);
+}
+
+void ULidarPointCloud::CalculateNormalsForSelection()
+{
+	Octree.CalculateNormalsForSelection(nullptr, NormalsQuality, NormalsNoiseTolerance);
+}
+
+void ULidarPointCloud::ClearSelection()
+{
+	Octree.ClearSelection();
+}
+#endif
+
 void ULidarPointCloud::SetSourcePath(const FString& NewSourcePath)
 {
 	SourcePath.FilePath = NewSourcePath;
@@ -585,6 +648,11 @@ void ULidarPointCloud::SetOptimizedForDynamicData(bool bNewOptimizedForDynamicDa
 	}
 }
 
+void ULidarPointCloud::SetOptimalCollisionError()
+{
+	MaxCollisionError = FMath::CeilToInt(Octree.GetEstimatedPointSpacing() * 300) * 0.01f;
+}
+
 void ULidarPointCloud::BuildCollision()
 {
 	if (bCollisionBuildInProgress)
@@ -758,7 +826,7 @@ void ULidarPointCloud::Reimport(const FLidarPointCloudAsyncParameters& AsyncPara
 				LocationOffset = bCenter ? FVector::ZeroVector : OriginalCoordinates;
 
 				// Adjust default max collision error
-				MaxCollisionError = FMath::CeilToInt(Octree.GetEstimatedPointSpacing() * 300) * 0.01f;
+				SetOptimalCollisionError();
 			}
 			else
 			{
@@ -1159,6 +1227,8 @@ void ULidarPointCloud::Merge(TArray<ULidarPointCloud*> PointCloudsToMerge, TFunc
 		ThreadResult.Get();
 	}
 
+	SetOptimalCollisionError();
+	
 	MarkPackageDirty();
 	OnPointCloudRebuiltEvent.Broadcast();
 }
diff --git a/Engine/Plugins/Enterprise/LidarPointCloud/Source/LidarPointCloudRuntime/Private/LidarPointCloudComponent.cpp b/Engine/Plugins/Enterprise/LidarPointCloud/Source/LidarPointCloudRuntime/Private/LidarPointCloudComponent.cpp
--- a/Engine/Plugins/Enterprise/LidarPointCloud/Source/LidarPointCloudRuntime/Private/LidarPointCloudComponent.cpp
+++ b/Engine/Plugins/Enterprise/LidarPointCloud/Source/LidarPointCloudRuntime/Private/LidarPointCloudComponent.cpp
@@ -135,6 +135,84 @@ void ULidarPointCloudComponent::PostPointCloudSet()
 	}
 }
 
+#if WITH_EDITOR
+void ULidarPointCloudComponent::SelectByConvexVolume(FConvexVolume ConvexVolume, bool bAdditive, bool bVisibleOnly)
+{
+	if(!PointCloud)
+	{
+		return;
+	}
+	
+	const FMatrix InvTransform = GetComponentTransform().Inverse().ToMatrixNoScale().ConcatTranslation(-PointCloud->LocationOffset);
+
+	for(FPlane& Plane : ConvexVolume.Planes)
+	{
+		Plane = Plane.TransformBy(InvTransform);
+	}
+
+	ConvexVolume.Init();
+
+	PointCloud->SelectByConvexVolume(ConvexVolume, bAdditive, false, bVisibleOnly);
+}
+
+void ULidarPointCloudComponent::SelectBySphere(FSphere Sphere, bool bAdditive, bool bVisibleOnly)
+{
+	if(!PointCloud)
+	{
+		return;
+	}
+	
+	const FMatrix InvTransform = GetComponentTransform().Inverse().ToMatrixNoScale().ConcatTranslation(-PointCloud->LocationOffset);
+	
+	PointCloud->SelectBySphere(Sphere.TransformBy(InvTransform), bAdditive, false, bVisibleOnly);
+}
+
+void ULidarPointCloudComponent::HideSelected()
+{
+	if(PointCloud)
+	{
+		PointCloud->HideSelected();
+	}
+}
+
+void ULidarPointCloudComponent::DeleteSelected()
+{
+	if(PointCloud)
+	{
+		PointCloud->DeleteSelected();
+	}
+}
+
+void ULidarPointCloudComponent::InvertSelection()
+{
+	if(PointCloud)
+	{
+		PointCloud->InvertSelection();
+	}
+}
+
+int64 ULidarPointCloudComponent::NumSelectedPoints()
+{
+	return PointCloud ? PointCloud->NumSelectedPoints() : 0;
+}
+
+void ULidarPointCloudComponent::GetSelectedPointsAsCopies(TArray64<FLidarPointCloudPoint>& SelectedPoints)
+{
+	if(PointCloud)
+	{
+		PointCloud->GetSelectedPointsAsCopies(SelectedPoints, GetComponentTransform());
+	}
+}
+
+void ULidarPointCloudComponent::ClearSelection()
+{
+	if(PointCloud)
+	{
+		PointCloud->ClearSelection();
+	}
+}
+#endif // WITH_EDITOR
+
 void ULidarPointCloudComponent::SetPointCloud(ULidarPointCloud *InPointCloud)
 {
 	if (PointCloud != InPointCloud)
diff --git a/Engine/Plugins/Enterprise/LidarPointCloud/Source/LidarPointCloudRuntime/Private/LidarPointCloudOctree.cpp b/Engine/Plugins/Enterprise/LidarPointCloud/Source/LidarPointCloudRuntime/Private/LidarPointCloudOctree.cpp
--- a/Engine/Plugins/Enterprise/LidarPointCloud/Source/LidarPointCloudRuntime/Private/LidarPointCloudOctree.cpp
+++ b/Engine/Plugins/Enterprise/LidarPointCloud/Source/LidarPointCloudRuntime/Private/LidarPointCloudOctree.cpp
@@ -119,6 +119,7 @@ FLidarPointCloudOctreeNode::FLidarPointCloudOctreeNode(FLidarPointCloudOctree* T
 	, Tree(Tree)
 	, bVisibilityDirty(false)
 	, bInUse(false)
+	, bHasSelection(false)
 	, NumVisiblePoints(0)
 	, NumPoints(0)
 	, bHasData(false)
@@ -787,7 +788,11 @@ int64 FLidarPointCloudOctreeNode::GetAllocatedSize(bool bRecursive, bool bInclud
 void FLidarPointCloudOctreeNode::ReleaseData(bool bForce)
 {
 	// Check if the data can be released
-	if (bCanReleaseData || bForce)
+	if ((bCanReleaseData
+#if WITH_EDITOR
+		&& !bHasSelection
+#endif
+		) || bForce)
 	{
 		bHasDataPending = false;
 		bCanReleaseData = true;
@@ -1151,6 +1156,7 @@ void FLidarPointCloudOctree::GetPointsAsCopies_Internal(TArray<FLidarPointCloudP
 
 	if (LocalToWorld)
 	{
+		const FTransform3f Transform = (FTransform3f)*LocalToWorld;
 		ITERATE_NODES_CONST({
 			// If no data is required, quit
 			if (Count == 0)
@@ -1166,7 +1172,7 @@ void FLidarPointCloudOctree::GetPointsAsCopies_Internal(TArray<FLidarPointCloudP
 				{
 					for (FLidarPointCloudPoint* Point = CurrentNode->GetData() + StartIndex, *DataEnd = Point + StartIndex + NumPointsToCopy; Point != DataEnd; ++Point)
 					{
-						Points.Add(Point->Transform(*LocalToWorld));
+						Points.Add(Point->Transform(Transform));
 					}
 
 					StartIndex = FMath::Max(0LL, StartIndex - NumPointsToCopy);
@@ -1327,7 +1333,8 @@ void FLidarPointCloudOctree::GetPointsInSphereAsCopies_Internal(TArray<FLidarPoi
 	SelectedPoints.Reset();
 	if (LocalToWorld)
 	{
-		PROCESS_IN_SPHERE_CONST({ SelectedPoints.Add(Point->Transform(*LocalToWorld)); });
+		const FTransform3f Transform = (FTransform3f)*LocalToWorld;
+		PROCESS_IN_SPHERE_CONST({ SelectedPoints.Add(Point->Transform(Transform)); });
 	}
 	else
 	{
@@ -1341,7 +1348,8 @@ void FLidarPointCloudOctree::GetPointsInBoxAsCopies_Internal(TArray<FLidarPointC
 	SelectedPoints.Reset();
 	if (LocalToWorld)
 	{
-		PROCESS_IN_BOX_CONST({ SelectedPoints.Add(Point->Transform(*LocalToWorld)); });
+		const FTransform3f Transform = (FTransform3f)*LocalToWorld;
+		PROCESS_IN_BOX_CONST({ SelectedPoints.Add(Point->Transform(Transform)); });
 	}
 	else
 	{
@@ -1387,7 +1395,8 @@ bool FLidarPointCloudOctree::RaycastMulti(const FLidarPointCloudRay& Ray, const
 	OutHits.Reset();
 	if (LocalToWorld)
 	{
-		PROCESS_BY_RAY_CONST({ OutHits.Add(Point->Transform(*LocalToWorld)); });
+		const FTransform3f Transform = (FTransform3f)*LocalToWorld;
+		PROCESS_BY_RAY_CONST({ OutHits.Add(Point->Transform(Transform)); });
 	}
 	else
 	{
@@ -1715,6 +1724,203 @@ void FLidarPointCloudOctree::MarkRenderDataInConvexVolumeDirty(const FConvexVolu
 	ITERATE_NODES({ CurrentNode->bRenderDataDirty = true; }, NODE_IN_CONVEX_VOLUME);
 }
 
+#if WITH_EDITOR
+void FLidarPointCloudOctree::SelectByConvexVolume(const FConvexVolume& ConvexVolume, bool bAdditive, bool bVisibleOnly)
+{
+	ITERATE_NODES(
+	{
+		// This will be used to determine if the payload should be released once convex check is complete
+		const bool bHadData = CurrentNode->bHasData;
+		bool bModified = false;
+
+		// Proactively flag as selected to avoid async release of the node's data
+		const bool bHadSelection = CurrentNode->bHasSelection;
+		CurrentNode->bHasSelection = true;
+
+		PROCESS_IN_CONVEX_VOLUME_BODY(
+		{
+			bModified = true;
+			Point->bSelected = bAdditive;
+		}, _RO)
+
+		if(bModified)
+		{
+			CurrentNode->bRenderDataDirty = true;
+		}
+		else
+		{
+			CurrentNode->bHasSelection = bHadSelection;
+
+			// Only attempt to release payload if the node didn't have any data loaded before
+			if(!bHadData)
+			{
+				CurrentNode->ReleaseData();
+			}
+		}
+	}, NODE_IN_CONVEX_VOLUME);
+}
+
+void FLidarPointCloudOctree::SelectBySphere(const FSphere& Sphere, bool bAdditive, bool bVisibleOnly)
+{
+	const FBox Box(Sphere.Center - FVector(Sphere.W), Sphere.Center + FVector(Sphere.W));
+	const float RadiusSq = Sphere.W * Sphere.W;
+	
+	ITERATE_NODES(
+	{
+		// This will be used to determine if the payload should be released once convex check is complete
+		const bool bHadData = CurrentNode->bHasData;
+		bool bModified = false;
+
+		// Proactively flag as selected to avoid async release of the node's data
+		const bool bHadSelection = CurrentNode->bHasSelection;
+		CurrentNode->bHasSelection = true;
+
+		PROCESS_IN_SPHERE_BODY(
+		{
+			bModified = true;
+			Point->bSelected = bAdditive;
+		}, _RO)
+
+		if(bModified)
+		{
+			CurrentNode->bRenderDataDirty = true;
+		}
+		else
+		{
+			CurrentNode->bHasSelection = bHadSelection;
+
+			// Only attempt to release payload if the node didn't have any data loaded before
+			if(!bHadData)
+			{
+				CurrentNode->ReleaseData();
+			}
+		}
+	}, NODE_IN_BOX);
+}
+
+void FLidarPointCloudOctree::HideSelected()
+{
+	ITERATE_SELECTED({
+		Point->bVisible = false;
+		Point->bSelected = false;
+	}, {
+		CurrentNode->bHasSelection = false;
+		CurrentNode->bCanReleaseData = false;
+		CurrentNode->bRenderDataDirty = true;
+		CurrentNode->bVisibilityDirty = true;
+	});
+}
+
+void FLidarPointCloudOctree::DeleteSelected()
+{
+	ITERATE_SELECTED_NODES({
+		const int32 NumElements = CurrentNode->GetNumPoints();
+		FLidarPointCloudPoint* Start = CurrentNode->GetData();
+		
+		int64 NumRemoved = 0;
+
+		for (FLidarPointCloudPoint* P = Start, *DataEnd = Start + NumElements; P != DataEnd; ++P)
+		{
+			if (P->bSelected)
+			{
+				CurrentNode->Data.RemoveAtSwap(P - Start, 1, false);
+				++NumRemoved;
+				--DataEnd;
+				--P;
+			}
+		}
+
+		CurrentNode->AddPointCount(-NumRemoved);
+
+		// Reduce space usage
+		CurrentNode->Data.Shrink();
+
+		// Copy the updated array back to the BulkData
+		CurrentNode->NumPoints = CurrentNode->Data.Num();
+
+		CurrentNode->bHasSelection = false;
+		CurrentNode->bCanReleaseData = false;
+		CurrentNode->bRenderDataDirty = true;
+		CurrentNode->bVisibilityDirty = true;
+	});
+}
+
+void FLidarPointCloudOctree::InvertSelection()
+{
+	ITERATE_NODES(
+	{
+		// This will be used to determine if the payload should be released once invert is complete
+		bool bNewHasSelection = false;
+		
+		FOR_RO(Point, CurrentNode)
+		{
+			Point->bSelected = !Point->bSelected;
+			if(Point->bSelected)
+			{
+				bNewHasSelection = true;
+			}
+		}
+
+		CurrentNode->bHasSelection = bNewHasSelection;
+		CurrentNode->bRenderDataDirty = true;
+
+		// Only attempt to release payload if the node isn't also used by rendering
+		if(!bNewHasSelection && !CurrentNode->bInUse)
+		{
+			CurrentNode->ReleaseData();
+		}
+	}, true);
+}
+
+int64 FLidarPointCloudOctree::NumSelectedPoints()
+{
+	int64 Count = 0;
+	ITERATE_SELECTED({ ++Count; }, {});
+	return Count;
+}
+
+void FLidarPointCloudOctree::GetSelectedPointsAsCopies(TArray64<FLidarPointCloudPoint>& SelectedPoints, const FTransform& Transform)
+{
+	const FTransform3f Transform3F = (FTransform3f)Transform;
+
+	ITERATE_SELECTED(
+	{
+		SelectedPoints.Add(Point->Transform(Transform3F));
+	},{});
+}
+
+void FLidarPointCloudOctree::CalculateNormalsForSelection(FThreadSafeBool* bCancelled, int32 Quality, float Tolerance)
+{
+	TArray64<FLidarPointCloudPoint*> SelectedPoints;
+
+	ITERATE_SELECTED({
+		SelectedPoints.Add(Point);
+		Point->Normal.Reset();
+	},{
+		CurrentNode->bCanReleaseData = false;
+		CurrentNode->bRenderDataDirty = true;
+	});
+	
+	LidarPointCloudMeshing::CalculateNormals(this, bCancelled, Quality, Tolerance, SelectedPoints);
+}
+
+void FLidarPointCloudOctree::ClearSelection()
+{
+	ITERATE_SELECTED({
+		Point->bSelected = false; 
+	}, {
+		CurrentNode->bRenderDataDirty = true;
+		CurrentNode->bHasSelection = false;
+
+		// Only attempt to release payload if the node isn't also used by rendering
+		if(!CurrentNode->bInUse)
+		{
+			CurrentNode->ReleaseData();
+		}
+	});
+}
+#endif // WITH_EDITOR
+
 void FLidarPointCloudOctree::MarkRenderDataInFrustumDirty(const FConvexVolume& Frustum)
 {
 	MarkRenderDataInConvexVolumeDirty(Frustum);
diff --git a/Engine/Plugins/Enterprise/LidarPointCloud/Source/LidarPointCloudRuntime/Private/LidarPointCloudOctreeMacros.h b/Engine/Plugins/Enterprise/LidarPointCloud/Source/LidarPointCloudRuntime/Private/LidarPointCloudOctreeMacros.h
--- a/Engine/Plugins/Enterprise/LidarPointCloud/Source/LidarPointCloudRuntime/Private/LidarPointCloudOctreeMacros.h
+++ b/Engine/Plugins/Enterprise/LidarPointCloud/Source/LidarPointCloudRuntime/Private/LidarPointCloudOctreeMacros.h
@@ -28,6 +28,19 @@
 #define ITERATE_NODES(Action, NodeTest) ITERATE_NODES_BODY(Action, NodeTest,  )
 #define ITERATE_NODES_CONST(Action, NodeTest) ITERATE_NODES_BODY(Action, NodeTest, const)
 
+#define ITERATE_SELECTED_NODES(NodeAction) \
+ITERATE_NODES(\
+{\
+	if(CurrentNode->bHasSelection) { NodeAction }\
+}, true);
+
+#define ITERATE_SELECTED(PointAction, NodeAction) \
+ITERATE_SELECTED_NODES(\
+{\
+	FOR_RO(Point, CurrentNode) { if(Point->bSelected) { PointAction } }\
+	{ NodeAction }\
+});
+
 #define POINT_IN_BOX Box.IsInsideOrOn((FVector)Point->Location)
 #define POINT_IN_SPHERE (POINT_IN_BOX && FVector3f::DistSquared(Point->Location, (FVector3f)Sphere.Center) <= RadiusSq)
 #define POINT_IN_CONVEX_VOLUME ConvexVolume.IntersectSphere((FVector)Point->Location, 0)
diff --git a/Engine/Plugins/Enterprise/LidarPointCloud/Source/LidarPointCloudRuntime/Public/LidarPointCloud.h b/Engine/Plugins/Enterprise/LidarPointCloud/Source/LidarPointCloudRuntime/Public/LidarPointCloud.h
--- a/Engine/Plugins/Enterprise/LidarPointCloud/Source/LidarPointCloudRuntime/Public/LidarPointCloud.h
+++ b/Engine/Plugins/Enterprise/LidarPointCloud/Source/LidarPointCloudRuntime/Public/LidarPointCloud.h
@@ -409,6 +409,19 @@ public:
 	UFUNCTION(BlueprintCallable, Category = "Lidar Point Cloud")
 	void MarkPointVisibilityDirty() { Octree.MarkPointVisibilityDirty(); }
 
+#if WITH_EDITOR
+	/** Set of editor helper functions */
+	void SelectByConvexVolume(FConvexVolume ConvexVolume, bool bAdditive, bool bApplyLocationOffset, bool bVisibleOnly);
+	void SelectBySphere(FSphere Sphere, bool bAdditive, bool bApplyLocationOffset, bool bVisibleOnly);
+	void HideSelected();
+	void DeleteSelected();
+	void InvertSelection();
+	int64 NumSelectedPoints();
+	void GetSelectedPointsAsCopies(TArray64<FLidarPointCloudPoint>& SelectedPoints, FTransform Transform);
+	void CalculateNormalsForSelection();
+	void ClearSelection();
+#endif
+	
 	UFUNCTION(BlueprintCallable, Category = "Lidar Point Cloud")
 	void SetSourcePath(const FString& NewSourcePath);
 
@@ -426,6 +439,9 @@ public:
 		Octree.Initialize((FVector3f)NewBounds.GetExtent());
 	}
 
+	UFUNCTION(BlueprintCallable, Category = "Lidar Point Cloud")
+	void SetOptimalCollisionError();
+
 	/** Builds collision mesh for the cloud, using current collision settings */
 	UFUNCTION(BlueprintCallable, Category = "Lidar Point Cloud")
 	void BuildCollision();
diff --git a/Engine/Plugins/Enterprise/LidarPointCloud/Source/LidarPointCloudRuntime/Public/LidarPointCloudComponent.h b/Engine/Plugins/Enterprise/LidarPointCloud/Source/LidarPointCloudRuntime/Public/LidarPointCloudComponent.h
--- a/Engine/Plugins/Enterprise/LidarPointCloud/Source/LidarPointCloudRuntime/Public/LidarPointCloudComponent.h
+++ b/Engine/Plugins/Enterprise/LidarPointCloud/Source/LidarPointCloudRuntime/Public/LidarPointCloudComponent.h
@@ -502,6 +502,17 @@ public:
 		}
 	}
 
+#if WITH_EDITOR
+	void SelectByConvexVolume(FConvexVolume ConvexVolume, bool bAdditive, bool bVisibleOnly);
+	void SelectBySphere(FSphere Sphere, bool bAdditive, bool bVisibleOnly);
+	void HideSelected();
+	void DeleteSelected();
+	void InvertSelection();
+	int64 NumSelectedPoints();
+	void GetSelectedPointsAsCopies(TArray64<FLidarPointCloudPoint>& SelectedPoints);
+	void ClearSelection();
+#endif
+	
 public:
 	UFUNCTION(BlueprintCallable, Category = "Components|LidarPointCloud")
 	void SetPointCloud(ULidarPointCloud *InPointCloud);
diff --git a/Engine/Plugins/Enterprise/LidarPointCloud/Source/LidarPointCloudRuntime/Public/LidarPointCloudOctree.h b/Engine/Plugins/Enterprise/LidarPointCloud/Source/LidarPointCloudRuntime/Public/LidarPointCloudOctree.h
--- a/Engine/Plugins/Enterprise/LidarPointCloud/Source/LidarPointCloudRuntime/Public/LidarPointCloudOctree.h
+++ b/Engine/Plugins/Enterprise/LidarPointCloud/Source/LidarPointCloudRuntime/Public/LidarPointCloudOctree.h
@@ -69,6 +69,9 @@ private:
 	/** Marks the node as being used for rendering */
 	bool bInUse;
 
+	/** Marks the node as containing active selection */
+	bool bHasSelection;
+
 	/** Stores the number of visible points */
 	uint32 NumVisiblePoints;
 
@@ -467,6 +470,18 @@ public:
 	/** Marks render data of all nodes within the given convex volume as dirty. */
 	void MarkRenderDataInConvexVolumeDirty(const FConvexVolume& ConvexVolume);
 
+#if WITH_EDITOR
+	void SelectByConvexVolume(const FConvexVolume& ConvexVolume, bool bAdditive, bool bVisibleOnly);
+	void SelectBySphere(const FSphere& Sphere, bool bAdditive, bool bVisibleOnly);
+	void HideSelected();
+	void DeleteSelected();
+	void InvertSelection();
+	int64 NumSelectedPoints();
+	void GetSelectedPointsAsCopies(TArray64<FLidarPointCloudPoint>& SelectedPoints, const FTransform& Transform);
+	void CalculateNormalsForSelection(FThreadSafeBool* bCancelled, int32 Quality, float Tolerance);
+	void ClearSelection();
+#endif
+
 	/** Initializes the Octree properties. */
 	void Initialize(const FVector3f& InExtent);
 
diff --git a/Engine/Plugins/Enterprise/LidarPointCloud/Source/LidarPointCloudRuntime/Public/LidarPointCloudShared.h b/Engine/Plugins/Enterprise/LidarPointCloud/Source/LidarPointCloudRuntime/Public/LidarPointCloudShared.h
--- a/Engine/Plugins/Enterprise/LidarPointCloud/Source/LidarPointCloudRuntime/Public/LidarPointCloudShared.h
+++ b/Engine/Plugins/Enterprise/LidarPointCloud/Source/LidarPointCloudRuntime/Public/LidarPointCloudShared.h
@@ -180,9 +180,9 @@ public:
 		ClassificationID = Other.ClassificationID;
 	}
 
-	FORCEINLINE FLidarPointCloudPoint Transform(const FTransform& Transform) const
+	FORCEINLINE FLidarPointCloudPoint Transform(const FTransform3f& Transform) const
 	{
-		return FLidarPointCloudPoint((FVector3f)Transform.TransformPosition((FVector)Location), Color, bVisible, ClassificationID);
+		return FLidarPointCloudPoint(Transform.TransformPosition(Location), Color, bVisible, ClassificationID);
 	}
 
 	bool operator==(const FLidarPointCloudPoint& P) const { return Location == P.Location && Color == P.Color && bVisible == P.bVisible && ClassificationID == P.ClassificationID && Normal == P.Normal; }
