Fixed erros when using Alignment tools

diff --git a/Engine/Plugins/Enterprise/LidarPointCloud/Source/LidarPointCloudEditor/Private/LidarPointCloudEditor.cpp b/Engine/Plugins/Enterprise/LidarPointCloud/Source/LidarPointCloudEditor/Private/LidarPointCloudEditor.cpp
--- a/Engine/Plugins/Enterprise/LidarPointCloud/Source/LidarPointCloudEditor/Private/LidarPointCloudEditor.cpp
+++ b/Engine/Plugins/Enterprise/LidarPointCloud/Source/LidarPointCloudEditor/Private/LidarPointCloudEditor.cpp
@@ -941,7 +941,7 @@ TArray<FAssetData> FLidarPointCloudEditor::SelectAssets(const FText& Title)
 	OpenAssetDialogConfig.DialogTitleOverride = Title;
 	OpenAssetDialogConfig.DefaultPath = "/Game";
 	OpenAssetDialogConfig.bAllowMultipleSelection = true;
-	OpenAssetDialogConfig.AssetClassNames.Emplace(*ULidarPointCloud::StaticClass()->GetName());
+	OpenAssetDialogConfig.AssetClassNames.Emplace("/Script/LidarPointCloudRuntime.LidarPointCloud");
 
 	FContentBrowserModule& ContentBrowserModule = FModuleManager::LoadModuleChecked<FContentBrowserModule>("ContentBrowser");
 	return ContentBrowserModule.Get().CreateModalOpenAssetDialog(OpenAssetDialogConfig);
diff --git a/Engine/Plugins/Enterprise/LidarPointCloud/Source/LidarPointCloudEditor/Private/LidarPointCloudEditorHelper.cpp b/Engine/Plugins/Enterprise/LidarPointCloud/Source/LidarPointCloudEditor/Private/LidarPointCloudEditorHelper.cpp
--- a/Engine/Plugins/Enterprise/LidarPointCloud/Source/LidarPointCloudEditor/Private/LidarPointCloudEditorHelper.cpp
+++ b/Engine/Plugins/Enterprise/LidarPointCloud/Source/LidarPointCloudEditor/Private/LidarPointCloudEditorHelper.cpp
@@ -26,7 +26,7 @@ namespace
 		FSaveAssetDialogConfig SaveAssetDialogConfig;
 		SaveAssetDialogConfig.DialogTitleOverride = LOCTEXT("SelectDestination", "Select Destination");
 		SaveAssetDialogConfig.DefaultPath = "/Game";
-		SaveAssetDialogConfig.AssetClassNames.Emplace("/Script/PointCloudRuntime.LidarPointCloud");
+		SaveAssetDialogConfig.AssetClassNames.Emplace("/Script/LidarPointCloudRuntime.LidarPointCloud");
 		SaveAssetDialogConfig.ExistingAssetPolicy = ESaveAssetDialogExistingAssetPolicy::AllowButWarn;
 
 		FContentBrowserModule& ContentBrowserModule = FModuleManager::LoadModuleChecked<FContentBrowserModule>("ContentBrowser");
