Fix non-unity
#rnx

diff --git a/Engine/Plugins/Enterprise/LidarPointCloud/Source/LidarPointCloudEditor/Private/LidarPointCloudEditorHelper.h b/Engine/Plugins/Enterprise/LidarPointCloud/Source/LidarPointCloudEditor/Private/LidarPointCloudEditorHelper.h
--- a/Engine/Plugins/Enterprise/LidarPointCloud/Source/LidarPointCloudEditor/Private/LidarPointCloudEditorHelper.h
+++ b/Engine/Plugins/Enterprise/LidarPointCloud/Source/LidarPointCloudEditor/Private/LidarPointCloudEditorHelper.h
@@ -6,7 +6,9 @@
 #include "LidarPointCloudShared.h"
 
 class ALidarPointCloudActor;
+class FEditorViewportClient;
 class ULidarPointCloud;
+struct FConvexVolume;
 
 enum class ELidarPointCloudSelectionMode : uint8
 {
