Deprecate old uses of ControllerId in any editor viewport clients.

#rnx

diff --git a/Engine/Plugins/Enterprise/LidarPointCloud/Source/LidarPointCloudEditor/Private/LidarPointCloudEditorViewportClient.cpp b/Engine/Plugins/Enterprise/LidarPointCloud/Source/LidarPointCloudEditor/Private/LidarPointCloudEditorViewportClient.cpp
--- a/Engine/Plugins/Enterprise/LidarPointCloud/Source/LidarPointCloudEditor/Private/LidarPointCloudEditorViewportClient.cpp
+++ b/Engine/Plugins/Enterprise/LidarPointCloud/Source/LidarPointCloudEditor/Private/LidarPointCloudEditorViewportClient.cpp
@@ -175,23 +175,23 @@ void FLidarPointCloudEditorViewportClient::Tick(float DeltaSeconds)
 	}
 }
 
-bool FLidarPointCloudEditorViewportClient::InputKey(FViewport* InViewport, int32 ControllerId, FKey Key, EInputEvent Event, float AmountDepressed, bool Gamepad)
+bool FLidarPointCloudEditorViewportClient::InputKey(const FInputKeyEventArgs& EventArgs)
 {
 	using UE::Lidar::Private::Editor::PolySnapDistanceSq;
 
 	bool bHandled = false;
 
-	const bool bAlt = Key == EKeys::LeftAlt || Key == EKeys::RightAlt;
-	const bool bCtrl = Key == EKeys::LeftControl || Key == EKeys::RightControl;
+	const bool bAlt = EventArgs.Key == EKeys::LeftAlt || EventArgs.Key == EKeys::RightAlt;
+	const bool bCtrl = EventArgs.Key == EKeys::LeftControl || EventArgs.Key == EKeys::RightControl;
 
 	if (TSharedPtr<FLidarPointCloudEditor> Editor = PointCloudEditorPtr.Pin())
 	{
 		// Edit Mode
 		if (Editor->IsEditMode())
 		{
-			if (Event == IE_Pressed)
+			if (EventArgs.Event == IE_Pressed)
 			{
-				if (Key == EKeys::Delete)
+				if (EventArgs.Key == EKeys::Delete)
 				{
 					if (Viewport->KeyState(EKeys::LeftShift) || Viewport->KeyState(EKeys::RightShift))
 					{
@@ -204,13 +204,13 @@ bool FLidarPointCloudEditorViewportClient::InputKey(FViewport* InViewport, int32
 
 					bHandled = true;
 				}
-				else if (Key == EKeys::Escape)
+				else if (EventArgs.Key == EKeys::Escape)
 				{
 					Editor->DeselectPoints();
 					SelectionPoints.Empty();
 					bHandled = true;
 				}
-				else if (Key == EKeys::Enter)
+				else if (EventArgs.Key == EKeys::Enter)
 				{
 					if (SelectionMethod == ELidarPointCloudSelectionMethod::Polygonal)
 					{
@@ -229,7 +229,7 @@ bool FLidarPointCloudEditorViewportClient::InputKey(FViewport* InViewport, int32
 					bHandled = true;
 					SelectionMode = ELidarPointCloudSelectionMode::Add;
 				}
-				else if (Key == EKeys::LeftMouseButton)
+				else if (EventArgs.Key == EKeys::LeftMouseButton)
 				{
 					bHandled = true;
 
@@ -237,13 +237,13 @@ bool FLidarPointCloudEditorViewportClient::InputKey(FViewport* InViewport, int32
 					if (SelectionMethod == ELidarPointCloudSelectionMethod::Box)
 					{
 						// Mark the cursor location for selection start
-						InViewport->GetMousePos(SelectionPoints[SelectionPoints.AddUninitialized()]);
+						EventArgs.Viewport->GetMousePos(SelectionPoints[SelectionPoints.AddUninitialized()]);
 					}
 					else if (SelectionMethod == ELidarPointCloudSelectionMethod::Polygonal)
 					{
 						// Add new lasso point
 						FIntPoint NewPoint;
-						InViewport->GetMousePos(NewPoint);
+						EventArgs.Viewport->GetMousePos(NewPoint);
 
 						// Don't allow duplicates
 						if (SelectionPoints.Num() == 0 || SelectionPoints.Last() != NewPoint)
@@ -268,7 +268,7 @@ bool FLidarPointCloudEditorViewportClient::InputKey(FViewport* InViewport, int32
 					}
 					else if (SelectionMethod == ELidarPointCloudSelectionMethod::Lasso)
 					{
-						InViewport->GetMousePos(SelectionPoints[SelectionPoints.AddUninitialized()]);
+						EventArgs.Viewport->GetMousePos(SelectionPoints[SelectionPoints.AddUninitialized()]);
 					}
 					else if (SelectionMethod == ELidarPointCloudSelectionMethod::Paint)
 					{
@@ -281,12 +281,12 @@ bool FLidarPointCloudEditorViewportClient::InputKey(FViewport* InViewport, int32
 					// Do not block ability to change camera speed
 					if (!Viewport->KeyState(EKeys::RightMouseButton))
 					{
-						if (Key == EKeys::MouseScrollUp)
+						if (EventArgs.Key == EKeys::MouseScrollUp)
 						{
 							PaintingRadius *= 1.1f;
 							bHandled = true;
 						}
-						else if (Key == EKeys::MouseScrollDown)
+						else if (EventArgs.Key == EKeys::MouseScrollDown)
 						{
 							PaintingRadius /= 1.1f;
 							bHandled = true;
@@ -294,21 +294,21 @@ bool FLidarPointCloudEditorViewportClient::InputKey(FViewport* InViewport, int32
 					}
 				}
 			}
-			else if (Event == IE_Released)
+			else if (EventArgs.Event == IE_Released)
 			{
 				if (bAlt || bCtrl)
 				{
 					SelectionMode = ELidarPointCloudSelectionMode::None;
 					bHandled = true;
 				}
-				else if (Key == EKeys::LeftMouseButton)
+				else if (EventArgs.Key == EKeys::LeftMouseButton)
 				{
 					bHandled = true;
 
 					if (SelectionMethod == ELidarPointCloudSelectionMethod::Box)
 					{
 						// Mark the cursor location for selection end
-						InViewport->GetMousePos(SelectionPoints[SelectionPoints.AddUninitialized()]);
+						EventArgs.Viewport->GetMousePos(SelectionPoints[SelectionPoints.AddUninitialized()]);
 						OnBoxSelectionEnd();
 						SelectionPoints.Empty();
 					}
@@ -329,19 +329,19 @@ bool FLidarPointCloudEditorViewportClient::InputKey(FViewport* InViewport, int32
 
 	if (!bHandled)
 	{
-		bHandled = FEditorViewportClient::InputKey(InViewport, ControllerId, Key, Event, AmountDepressed, false);
+		bHandled = FEditorViewportClient::InputKey(EventArgs);
 
 		// Handle viewport screenshot.
-		bHandled |= InputTakeScreenshot(InViewport, Key, Event);
+		bHandled |= InputTakeScreenshot(EventArgs.Viewport, EventArgs.Key, EventArgs.Event);
 
-		bHandled |= AdvancedPreviewScene->HandleInputKey(InViewport, ControllerId, Key, Event, AmountDepressed, Gamepad);
+		bHandled |= AdvancedPreviewScene->HandleInputKey(EventArgs);
 
 	}
 
 	return bHandled;
 }
 
-bool FLidarPointCloudEditorViewportClient::InputAxis(FViewport* InViewport, int32 ControllerId, FKey Key, float Delta, float DeltaTime, int32 NumSamples, bool bGamepad)
+bool FLidarPointCloudEditorViewportClient::InputAxis(FViewport* InViewport, FInputDeviceId DeviceId, FKey Key, float Delta, float DeltaTime, int32 NumSamples, bool bGamepad)
 {
 	using UE::Lidar::Private::Editor::LassoSpacingSq;
 
@@ -379,14 +379,14 @@ bool FLidarPointCloudEditorViewportClient::InputAxis(FViewport* InViewport, int3
 
 		if(!bHandled)
 		{
-			bHandled = AdvancedPreviewScene->HandleViewportInput(InViewport, ControllerId, Key, Delta, DeltaTime, NumSamples, bGamepad);
+			bHandled = AdvancedPreviewScene->HandleViewportInput(InViewport, DeviceId, Key, Delta, DeltaTime, NumSamples, bGamepad);
 			if (bHandled)
 			{
 				Invalidate();
 			}
 			else
 			{
-				bHandled = FEditorViewportClient::InputAxis(InViewport, ControllerId, Key, Delta, DeltaTime, NumSamples, bGamepad);
+				bHandled = FEditorViewportClient::InputAxis(InViewport, DeviceId, Key, Delta, DeltaTime, NumSamples, bGamepad);
 			}
 		}
 	}
diff --git a/Engine/Plugins/Enterprise/LidarPointCloud/Source/LidarPointCloudEditor/Private/LidarPointCloudEditorViewportClient.h b/Engine/Plugins/Enterprise/LidarPointCloud/Source/LidarPointCloudEditor/Private/LidarPointCloudEditorViewportClient.h
--- a/Engine/Plugins/Enterprise/LidarPointCloud/Source/LidarPointCloudEditor/Private/LidarPointCloudEditorViewportClient.h
+++ b/Engine/Plugins/Enterprise/LidarPointCloud/Source/LidarPointCloudEditor/Private/LidarPointCloudEditorViewportClient.h
@@ -32,8 +32,8 @@ public:
 	~FLidarPointCloudEditorViewportClient();
 
 	// FEditorViewportClient interface
-	virtual bool InputKey(FViewport* InViewport, int32 ControllerId, FKey Key, EInputEvent Event, float AmountDepressed = 1.f, bool bGamepad = false) override;
-	virtual bool InputAxis(FViewport* InViewport, int32 ControllerId, FKey Key, float Delta, float DeltaTime, int32 NumSamples = 1, bool bGamepad = false) override;
+	virtual bool InputKey(const FInputKeyEventArgs& EventArgs) override;
+	virtual bool InputAxis(FViewport* InViewport, FInputDeviceId DeviceId, FKey Key, float Delta, float DeltaTime, int32 NumSamples = 1, bool bGamepad = false) override;
 	virtual void Tick(float DeltaSeconds) override;
 	virtual void DrawCanvas(FViewport& InViewport, FSceneView& View, FCanvas& Canvas) override;
 	virtual bool InputWidgetDelta(FViewport* InViewport, EAxisList::Type CurrentAxis, FVector& Drag, FRotator& Rot, FVector& Scale) override { return false; }
diff --git a/Engine/Source/Editor/AdvancedPreviewScene/Private/AdvancedPreviewScene.cpp b/Engine/Source/Editor/AdvancedPreviewScene/Private/AdvancedPreviewScene.cpp
--- a/Engine/Source/Editor/AdvancedPreviewScene/Private/AdvancedPreviewScene.cpp
+++ b/Engine/Source/Editor/AdvancedPreviewScene/Private/AdvancedPreviewScene.cpp
@@ -275,16 +275,14 @@
 	return TStatId();
 }
 
-const bool FAdvancedPreviewScene::HandleViewportInput(FViewport* InViewport, int32 ControllerId, FKey Key, float Delta, float DeltaTime, int32 NumSamples, bool bGamepad)
+const bool FAdvancedPreviewScene::HandleViewportInput(FViewport* InViewport, FInputDeviceId DeviceId, FKey Key, float Delta, float DeltaTime, int32 NumSamples, bool bGamepad)
 {
 	bool bResult = false;
 	const bool bMouseButtonDown = InViewport->KeyState(EKeys::LeftMouseButton) || InViewport->KeyState(EKeys::MiddleMouseButton) || InViewport->KeyState(EKeys::RightMouseButton);
 	
-	
 	const bool bSkyMove = InViewport->KeyState(EKeys::K);
 	const bool bLightMoveDown = InViewport->KeyState(EKeys::L);
 	
-
 	// Look at which axis is being dragged and by how much
 	const float DragX = (Key == EKeys::MouseX) ? Delta : 0.f;	
 	const float DragY = (Key == EKeys::MouseY) ? Delta : 0.f;
@@ -308,12 +306,25 @@
 	return bResult;
 }
 
+PRAGMA_DISABLE_DEPRECATION_WARNINGS
+const bool FAdvancedPreviewScene::HandleViewportInput(FViewport* InViewport, int32 ControllerId, FKey Key, float Delta, float DeltaTime, int32 NumSamples, bool bGamepad)
+{
+	return HandleViewportInput(InViewport, FInputDeviceId::CreateFromInternalId(ControllerId), Key, Delta, DeltaTime, NumSamples, bGamepad);
+}
+
 const bool FAdvancedPreviewScene::HandleInputKey(FViewport* InViewport, int32 ControllerId, FKey Key, EInputEvent Event, float AmountDepressed, bool Gamepad)
 {
-	if (Event == IE_Pressed)
+	FInputKeyEventArgs Args(InViewport, ControllerId, Key, Event, AmountDepressed, false);
+	return HandleInputKey(Args);
+}
+PRAGMA_ENABLE_DEPRECATION_WARNINGS
+
+const bool FAdvancedPreviewScene::HandleInputKey(const FInputKeyEventArgs& EventArgs)
+{
+	if (EventArgs.Event == IE_Pressed)
 	{
 		FModifierKeysState KeyState = FSlateApplication::Get().GetModifierKeys();
-		if (UICommandList->ProcessCommandBindings(Key, KeyState, (Event == IE_Repeat))) //-V547
+		if (UICommandList->ProcessCommandBindings(EventArgs.Key, KeyState, (EventArgs.Event == IE_Repeat))) //-V547
 		{
 			return true;
 		}
diff --git a/Engine/Source/Editor/AdvancedPreviewScene/Public/AdvancedPreviewScene.h b/Engine/Source/Editor/AdvancedPreviewScene/Public/AdvancedPreviewScene.h
--- a/Engine/Source/Editor/AdvancedPreviewScene/Public/AdvancedPreviewScene.h
+++ b/Engine/Source/Editor/AdvancedPreviewScene/Public/AdvancedPreviewScene.h
@@ -7,6 +7,7 @@
 #pragma once
 
 #include "CoreMinimal.h"
+#include "UnrealClient.h"
 #include "Stats/Stats.h"
 #include "InputCoreTypes.h"
 #include "PreviewScene.h"
@@ -40,9 +41,16 @@ public:
 	virtual TStatId GetStatId() const override;
 	/* End FTickableEditorObject */
 
+	UE_DEPRECATED(5.1, "This version of HandleInputKey is deprecated. Please use the version that takes EventArgs instead.")
 	const bool HandleViewportInput(FViewport* InViewport, int32 ControllerId, FKey Key, float Delta, float DeltaTime, int32 NumSamples, bool bGamepad);
+	
+	const bool HandleViewportInput(FViewport* InViewport, FInputDeviceId DeviceId, FKey Key, float Delta, float DeltaTime, int32 NumSamples, bool bGamepad);
+
+	UE_DEPRECATED(5.1, "This version of HandleInputKey is deprecated. Please use the version that takes EventArgs instead.")
 	const bool HandleInputKey(FViewport* InViewport, int32 ControllerId, FKey Key, EInputEvent Event, float AmountDepressed, bool Gamepad);
 
+	const bool HandleInputKey(const FInputKeyEventArgs& EventArgs);
+
 	void SetSkyRotation(const float SkyRotation);
 	/* Sets the visiblity state for the floor/environment by storing it in the scene profile and refreshing the scene, in case bDirect is true it sets the visibility directly and leaves the profile untouched. */
 	bool GetFloorVisibility() const;
diff --git a/Engine/Source/Editor/Cascade/Private/CascadePreviewViewportClient.cpp b/Engine/Source/Editor/Cascade/Private/CascadePreviewViewportClient.cpp
--- a/Engine/Source/Editor/Cascade/Private/CascadePreviewViewportClient.cpp
+++ b/Engine/Source/Editor/Cascade/Private/CascadePreviewViewportClient.cpp
@@ -590,21 +590,21 @@ void FCascadeEdPreviewViewportClient::Draw(const FSceneView* View, FPrimitiveDra
 	DrawPreviewLightVisualization(View, PDI);
 }
 
-bool FCascadeEdPreviewViewportClient::InputKey(FViewport* InViewport, int32 ControllerId, FKey Key, EInputEvent Event, float AmountDepressed, bool Gamepad)
+bool FCascadeEdPreviewViewportClient::InputKey(const FInputKeyEventArgs& InEventArgs)
 {
 	//Update cursor 
 	UpdateAndApplyCursorVisibility();
 
 	bool bHandled = false;
-	const int32 HitX = InViewport->GetMouseX();
-	const int32 HitY = InViewport->GetMouseY();
+	const int32 HitX = InEventArgs.Viewport->GetMouseX();
+	const int32 HitY = InEventArgs.Viewport->GetMouseY();
 
-	if(Key == EKeys::LeftMouseButton)
+	if(InEventArgs.Key == EKeys::LeftMouseButton)
 	{
-		if (Event == IE_Pressed)
+		if (InEventArgs.Event == IE_Pressed)
 		{
-			InViewport->InvalidateHitProxy();
-			HHitProxy* HitResult = InViewport->GetHitProxy(HitX,HitY);
+			InEventArgs.Viewport->InvalidateHitProxy();
+			HHitProxy* HitResult = InEventArgs.Viewport->GetHitProxy(HitX,HitY);
 			if (HitResult && HitResult->IsA(HWidgetUtilProxy::StaticGetType()))
 			{
 				HWidgetUtilProxy* WidgetProxy = (HWidgetUtilProxy*)HitResult;
@@ -615,13 +615,13 @@ bool FCascadeEdPreviewViewportClient::InputKey(FViewport* InViewport, int32 Cont
 				WidgetAxis = WidgetProxy->Axis;
 
 				// Calculate the scree-space directions for this drag.
-				FSceneViewFamilyContext ViewFamily( FSceneViewFamily::ConstructionValues( InViewport, GetScene(), EngineShowFlags ));
+				FSceneViewFamilyContext ViewFamily( FSceneViewFamily::ConstructionValues( InEventArgs.Viewport, GetScene(), EngineShowFlags ));
 				FSceneView* View = CalcSceneView(&ViewFamily);
-				WidgetProxy->CalcVectors(View, FViewportClick(View, this, Key, Event, HitX, HitY), LocalManipulateDir, WorldManipulateDir, DragX, DragY);
+				WidgetProxy->CalcVectors(View, FViewportClick(View, this, InEventArgs.Key, InEventArgs.Event, HitX, HitY), LocalManipulateDir, WorldManipulateDir, DragX, DragY);
 				bHandled = true;
 			}
 		}
-		else if (Event == IE_Released)
+		else if (InEventArgs.Event == IE_Released)
 		{
 			if (bManipulatingVectorField)
 			{
@@ -632,7 +632,7 @@ bool FCascadeEdPreviewViewportClient::InputKey(FViewport* InViewport, int32 Cont
 			}
 		}
 	}
-	else if (Key == EKeys::SpaceBar && Event == IE_Pressed)
+	else if (InEventArgs.Key == EKeys::SpaceBar && InEventArgs.Event == IE_Pressed)
 	{
 		if (CascadePtr.Pin()->GetSelectedModule() && CascadePtr.Pin()->GetSelectedModule()->IsA(UParticleModuleVectorFieldLocal::StaticClass()))
 		{
@@ -644,7 +644,7 @@ bool FCascadeEdPreviewViewportClient::InputKey(FViewport* InViewport, int32 Cont
 
 	if( !bHandled )
 	{
-		bHandled = FEditorViewportClient::InputKey(InViewport,ControllerId,Key,Event,AmountDepressed,Gamepad);
+		bHandled = FEditorViewportClient::InputKey(InEventArgs);
 	}
 
 
@@ -657,7 +657,7 @@ static auto CVarCascadeDragSpeed = TAutoConsoleVariable<float>(TEXT("CascadeDrag
 static TAutoConsoleVariable<float> CVarCascadeRotateSpeed(TEXT("CascadeRotateSpeed"),0.005f,TEXT("Cascade drag speed."));
 static TAutoConsoleVariable<float> CVarCascadeScaleSpeed(TEXT("CascadeScaleSpeed"),1.0f,TEXT("Cascade scale speed."));
 
-bool FCascadeEdPreviewViewportClient::InputAxis(FViewport* InViewport, int32 ControllerId, FKey Key, float Delta, float DeltaTime, int32 NumSamples, bool bGamepad)
+bool FCascadeEdPreviewViewportClient::InputAxis(FViewport* InViewport, FInputDeviceId DeviceId, FKey Key, float Delta, float DeltaTime, int32 NumSamples, bool bGamepad)
 {
 	bool bHandled = false;
 
@@ -693,7 +693,7 @@ bool FCascadeEdPreviewViewportClient::InputAxis(FViewport* InViewport, int32 Con
 	}
 	else
 	{
-		bHandled = FEditorViewportClient::InputAxis(InViewport,ControllerId,Key,Delta,DeltaTime,NumSamples,bGamepad);
+		bHandled = FEditorViewportClient::InputAxis(InViewport,DeviceId,Key,Delta,DeltaTime,NumSamples,bGamepad);
 	}
 
 	if (!IsRealtime() && !FMath::IsNearlyZero(Delta))
diff --git a/Engine/Source/Editor/Cascade/Private/CascadePreviewViewportClient.h b/Engine/Source/Editor/Cascade/Private/CascadePreviewViewportClient.h
--- a/Engine/Source/Editor/Cascade/Private/CascadePreviewViewportClient.h
+++ b/Engine/Source/Editor/Cascade/Private/CascadePreviewViewportClient.h
@@ -27,8 +27,8 @@ public:
 	/** FEditorViewportClient interface */
 	virtual void Draw(FViewport* Viewport, FCanvas* Canvas) override;
 	virtual void Draw(const FSceneView* View, FPrimitiveDrawInterface* PDI) override;
-	virtual bool InputKey(FViewport* Viewport, int32 ControllerId, FKey Key, EInputEvent Event, float AmountDepressed = 1.0f, bool bGamepad = false) override;
-	virtual bool InputAxis(FViewport* Viewport, int32 ControllerId, FKey Key, float Delta, float DeltaTime, int32 NumSamples = 1, bool bGamepad = false) override;
+	virtual bool InputKey(const FInputKeyEventArgs& InEventArgs) override;
+	virtual bool InputAxis(FViewport* Viewport, FInputDeviceId DeviceId, FKey Key, float Delta, float DeltaTime, int32 NumSamples = 1, bool bGamepad = false) override;
 	virtual FSceneInterface* GetScene() const override;
 	virtual FLinearColor GetBackgroundColor() const override;
 	virtual bool ShouldOrbitCamera() const override;
diff --git a/Engine/Source/Editor/MaterialEditor/Private/SMaterialEditorViewport.cpp b/Engine/Source/Editor/MaterialEditor/Private/SMaterialEditorViewport.cpp
--- a/Engine/Source/Editor/MaterialEditor/Private/SMaterialEditorViewport.cpp
+++ b/Engine/Source/Editor/MaterialEditor/Private/SMaterialEditorViewport.cpp
@@ -48,8 +48,8 @@ public:
 	FMaterialEditorViewportClient(TWeakPtr<IMaterialEditor> InMaterialEditor, FAdvancedPreviewScene& InPreviewScene, const TSharedRef<SMaterialEditor3DPreviewViewport>& InMaterialEditorViewport);
 
 	// FEditorViewportClient interface
-	virtual bool InputKey(FViewport* InViewport, int32 ControllerId, FKey Key, EInputEvent Event, float AmountDepressed = 1.f, bool bGamepad = false) override;
-	virtual bool InputAxis(FViewport* InViewport, int32 ControllerId, FKey Key, float Delta, float DeltaTime, int32 NumSamples/* =1 */, bool bGamepad/* =false */) override;
+	virtual bool InputKey(const FInputKeyEventArgs& EventArgs) override;
+	virtual bool InputAxis(FViewport* InViewport, FInputDeviceId DeviceId, FKey Key, float Delta, float DeltaTime, int32 NumSamples/* =1 */, bool bGamepad/* =false */) override;
 	virtual FLinearColor GetBackgroundColor() const override;
 	virtual void Tick(float DeltaSeconds) override;
 	virtual void Draw(FViewport* Viewport,FCanvas* Canvas) override;
@@ -128,32 +128,32 @@ bool FMaterialEditorViewportClient::ShouldOrbitCamera() const
 	return true;
 }
 
-bool FMaterialEditorViewportClient::InputKey(FViewport* InViewport, int32 ControllerId, FKey Key, EInputEvent Event, float AmountDepressed, bool bGamepad)
+bool FMaterialEditorViewportClient::InputKey(const FInputKeyEventArgs& EventArgs)
 {
-	bool bHandled = FEditorViewportClient::InputKey(InViewport, ControllerId, Key, Event, AmountDepressed, false);
+	bool bHandled = FEditorViewportClient::InputKey(EventArgs);
 
 	// Handle viewport screenshot.
-	bHandled |= InputTakeScreenshot(InViewport, Key, Event);
+	bHandled |= InputTakeScreenshot(EventArgs.Viewport, EventArgs.Key, EventArgs.Event);
 
-	bHandled |= AdvancedPreviewScene->HandleInputKey(InViewport, ControllerId, Key, Event, AmountDepressed, bGamepad);
+	bHandled |= AdvancedPreviewScene->HandleInputKey(EventArgs);
 
 	return bHandled;
 }
 
-bool FMaterialEditorViewportClient::InputAxis(FViewport* InViewport, int32 ControllerId, FKey Key, float Delta, float DeltaTime, int32 NumSamples/* =1 */, bool bGamepad/* =false */)
+bool FMaterialEditorViewportClient::InputAxis(FViewport* InViewport, FInputDeviceId DeviceId, FKey Key, float Delta, float DeltaTime, int32 NumSamples/* =1 */, bool bGamepad/* =false */)
 {
 	bool bResult = true;
 
 	if (!bDisableInput)
 	{
-		bResult = AdvancedPreviewScene->HandleViewportInput(InViewport, ControllerId, Key, Delta, DeltaTime, NumSamples, bGamepad);
+		bResult = AdvancedPreviewScene->HandleViewportInput(InViewport, DeviceId, Key, Delta, DeltaTime, NumSamples, bGamepad);
 		if (bResult)
 		{
 			Invalidate();
 		}
 		else
 		{
-			bResult = FEditorViewportClient::InputAxis(InViewport, ControllerId, Key, Delta, DeltaTime, NumSamples, bGamepad);
+			bResult = FEditorViewportClient::InputAxis(InViewport, DeviceId, Key, Delta, DeltaTime, NumSamples, bGamepad);
 		}
 	}
 
diff --git a/Engine/Source/Editor/Persona/Private/AnimationEditorViewportClient.cpp b/Engine/Source/Editor/Persona/Private/AnimationEditorViewportClient.cpp
--- a/Engine/Source/Editor/Persona/Private/AnimationEditorViewportClient.cpp
+++ b/Engine/Source/Editor/Persona/Private/AnimationEditorViewportClient.cpp
@@ -18,6 +18,7 @@
 #include "GameFramework/WorldSettings.h"
 #include "PersonaModule.h"
 #include "Rendering/SkeletalMeshRenderData.h"
+#include "GenericPlatform/GenericPlatformInputDeviceMapper.h"
 
 #include "SEditorViewport.h"
 #include "CanvasTypes.h"
@@ -1314,41 +1315,57 @@ void FAnimationViewportClient::RotateViewportType()
 		bUsingOrbitCamera = true;
 	}
 }
-
-bool FAnimationViewportClient::InputKey( FViewport* InViewport, int32 ControllerId, FKey Key, EInputEvent Event, float AmountDepressed, bool bGamepad )
+bool FAnimationViewportClient::InputKey(const FInputKeyEventArgs& EventArgs)
 {
 	bool bHandled = false;
 
 	FAdvancedPreviewScene* AdvancedScene = static_cast<FAdvancedPreviewScene*>(PreviewScene);
-	bHandled |= AdvancedScene->HandleInputKey(InViewport, ControllerId, Key, Event, AmountDepressed, bGamepad);
+	bHandled |= AdvancedScene->HandleInputKey(EventArgs);
 
 	// Pass keys to standard controls, if we didn't consume input
 	return (bHandled)
 		? true
-		: FEditorViewportClient::InputKey(InViewport, ControllerId, Key, Event, AmountDepressed, bGamepad);
+		: FEditorViewportClient::InputKey(EventArgs);
 }
 
-bool FAnimationViewportClient::InputAxis(FViewport* InViewport, int32 ControllerId, FKey Key, float Delta, float DeltaTime, int32 NumSamples /*= 1*/, bool bGamepad /*= false*/)
+bool FAnimationViewportClient::InputAxis(FViewport* InViewport, FInputDeviceId DeviceId, FKey Key, float Delta, float DeltaTime, int32 NumSamples /*= 1*/, bool bGamepad /*= false*/)
 {
 	bool bResult = true;
 
 	if (!bDisableInput)
 	{
 		FAdvancedPreviewScene* AdvancedScene = (FAdvancedPreviewScene*)PreviewScene;
-		bResult = AdvancedScene->HandleViewportInput(InViewport, ControllerId, Key, Delta, DeltaTime, NumSamples, bGamepad);
+		bResult = AdvancedScene->HandleViewportInput(InViewport, DeviceId, Key, Delta, DeltaTime, NumSamples, bGamepad);
 		if (bResult)
 		{
 			Invalidate();
 		}
 		else
 		{
-			bResult = FEditorViewportClient::InputAxis(InViewport, ControllerId, Key, Delta, DeltaTime, NumSamples, bGamepad);
+			bResult = FEditorViewportClient::InputAxis(InViewport, DeviceId, Key, Delta, DeltaTime, NumSamples, bGamepad);
 		}
 	}
 
 	return bResult;
 }
 
+PRAGMA_DISABLE_DEPRECATION_WARNINGS
+bool FAnimationViewportClient::InputKey( FViewport* InViewport, int32 ControllerId, FKey Key, EInputEvent Event, float AmountDepressed, bool bGamepad )
+{
+	FInputKeyEventArgs Args(InViewport, ControllerId, Key, Event, AmountDepressed, /*bIsTouchEvent*/false);
+	return InputKey(Args);
+}
+
+bool FAnimationViewportClient::InputAxis(FViewport* InViewport, int32 ControllerId, FKey Key, float Delta, float DeltaTime, int32 NumSamples /*= 1*/, bool bGamepad /*= false*/)
+{
+	FPlatformUserId UserId = FGenericPlatformMisc::GetPlatformUserForUserIndex(ControllerId);
+	FInputDeviceId DeviceId = INPUTDEVICEID_NONE;
+	IPlatformInputDeviceMapper::Get().RemapControllerIdToPlatformUserAndDevice(ControllerId, UserId, DeviceId);
+	
+	return InputAxis(InViewport, DeviceId, Key, Delta, DeltaTime, NumSamples, bGamepad);
+}
+PRAGMA_ENABLE_DEPRECATION_WARNINGS
+
 void FAnimationViewportClient::SetLocalAxesMode(ELocalAxesMode::Type AxesMode)
 {
 	ConfigOption->SetDefaultLocalAxesSelection( AxesMode );
diff --git a/Engine/Source/Editor/Persona/Private/AnimationEditorViewportClient.h b/Engine/Source/Editor/Persona/Private/AnimationEditorViewportClient.h
--- a/Engine/Source/Editor/Persona/Private/AnimationEditorViewportClient.h
+++ b/Engine/Source/Editor/Persona/Private/AnimationEditorViewportClient.h
@@ -110,8 +110,15 @@ public:
 	virtual void Tick(float DeltaSeconds) override;
 	virtual void Draw(const FSceneView* View,FPrimitiveDrawInterface* PDI) override;
 	virtual void DrawCanvas( FViewport& InViewport, FSceneView& View, FCanvas& Canvas ) override;
+	
+	UE_DEPRECATED(5.1, "This version of InputKey is deprecated. Please use the version that takes EventArgs instead.")
 	virtual bool InputKey(FViewport* Viewport, int32 ControllerId, FKey Key, EInputEvent Event, float AmountDepressed = 1.f, bool bGamepad=false) override;
+	virtual bool InputKey(const FInputKeyEventArgs& EventArgs) override;
+	
+	UE_DEPRECATED(5.1, "This version of InputAxis is deprecated. Please use the version that takes DeviceId instead.")
 	virtual bool InputAxis(FViewport* InViewport, int32 ControllerId, FKey Key, float Delta, float DeltaTime, int32 NumSamples = 1, bool bGamepad = false) override;
+	virtual bool InputAxis(FViewport* InViewport, FInputDeviceId DeviceId, FKey Key, float Delta, float DeltaTime, int32 NumSamples = 1, bool bGamepad = false) override;
+	
 //	virtual void ProcessClick(class FSceneView& View, class HHitProxy* HitProxy, FKey Key, EInputEvent Event, uint32 HitX, uint32 HitY) override;
 //	virtual bool InputWidgetDelta( FViewport* Viewport, EAxisList::Type CurrentAxis, FVector& Drag, FRotator& Rot, FVector& Scale ) override;
 	virtual void TrackingStarted( const struct FInputEventState& InInputState, bool bIsDragging, bool bNudge ) override;
diff --git a/Engine/Source/Editor/StaticMeshEditor/Private/StaticMeshEditorViewportClient.cpp b/Engine/Source/Editor/StaticMeshEditor/Private/StaticMeshEditorViewportClient.cpp
--- a/Engine/Source/Editor/StaticMeshEditor/Private/StaticMeshEditorViewportClient.cpp
+++ b/Engine/Source/Editor/StaticMeshEditor/Private/StaticMeshEditorViewportClient.cpp
@@ -961,32 +961,32 @@ void FStaticMeshEditorViewportClient::MouseMove(FViewport* InViewport,int32 x, i
 	FEditorViewportClient::MouseMove(InViewport,x,y);
 }
 
-bool FStaticMeshEditorViewportClient::InputKey(FViewport* InViewport, int32 ControllerId, FKey Key, EInputEvent Event,float AmountDepressed, bool Gamepad)
+bool FStaticMeshEditorViewportClient::InputKey(const FInputKeyEventArgs& EventArgs)
 {
-	bool bHandled = FEditorViewportClient::InputKey(InViewport, ControllerId, Key, Event, AmountDepressed, false);
+	bool bHandled = FEditorViewportClient::InputKey(EventArgs);
 
 	// Handle viewport screenshot.
-	bHandled |= InputTakeScreenshot( InViewport, Key, Event );
+	bHandled |= InputTakeScreenshot( EventArgs.Viewport, EventArgs.Key, EventArgs.Event );
 
-	bHandled |= AdvancedPreviewScene->HandleInputKey(InViewport, ControllerId, Key, Event, AmountDepressed, Gamepad);
+	bHandled |= AdvancedPreviewScene->HandleInputKey(EventArgs);
 
 	return bHandled;
 }
 
-bool FStaticMeshEditorViewportClient::InputAxis(FViewport* InViewport, int32 ControllerId, FKey Key, float Delta, float DeltaTime, int32 NumSamples, bool bGamepad)
+bool FStaticMeshEditorViewportClient::InputAxis(FViewport* InViewport, FInputDeviceId DeviceId, FKey Key, float Delta, float DeltaTime, int32 NumSamples, bool bGamepad)
 {
 	bool bResult = true;
 	
 	if (!bDisableInput)
 	{
-		bResult = AdvancedPreviewScene->HandleViewportInput(InViewport, ControllerId, Key, Delta, DeltaTime, NumSamples, bGamepad);
+		bResult = AdvancedPreviewScene->HandleViewportInput(InViewport, DeviceId, Key, Delta, DeltaTime, NumSamples, bGamepad);
 		if (bResult)
 		{
 			Invalidate();
 		}
 		else
 		{
-			bResult = FEditorViewportClient::InputAxis(InViewport, ControllerId, Key, Delta, DeltaTime, NumSamples, bGamepad);
+			bResult = FEditorViewportClient::InputAxis(InViewport, DeviceId, Key, Delta, DeltaTime, NumSamples, bGamepad);
 		}
 	}
 
diff --git a/Engine/Source/Editor/StaticMeshEditor/Private/StaticMeshEditorViewportClient.h b/Engine/Source/Editor/StaticMeshEditor/Private/StaticMeshEditorViewportClient.h
--- a/Engine/Source/Editor/StaticMeshEditor/Private/StaticMeshEditorViewportClient.h
+++ b/Engine/Source/Editor/StaticMeshEditor/Private/StaticMeshEditorViewportClient.h
@@ -26,8 +26,8 @@ public:
 
 	// FEditorViewportClient interface
 	virtual void MouseMove(FViewport* Viewport,int32 x, int32 y) override;
-	virtual bool InputKey(FViewport* Viewport, int32 ControllerId, FKey Key, EInputEvent Event, float AmountDepressed = 1.f, bool bGamepad=false) override;
-	virtual bool InputAxis(FViewport* Viewport, int32 ControllerId, FKey Key, float Delta, float DeltaTime, int32 NumSamples=1, bool bGamepad=false) override;
+	virtual bool InputKey(const FInputKeyEventArgs& EventArgs) override;
+	virtual bool InputAxis(FViewport* Viewport, FInputDeviceId DeviceId, FKey Key, float Delta, float DeltaTime, int32 NumSamples=1, bool bGamepad=false) override;
 	virtual void ProcessClick(class FSceneView& View, class HHitProxy* HitProxy, FKey Key, EInputEvent Event, uint32 HitX, uint32 HitY) override;
 	virtual void Tick(float DeltaSeconds) override;
 	virtual void Draw(const FSceneView* View,FPrimitiveDrawInterface* PDI) override;
