Ability to create actors/meshes from selected lidar points in the world

diff --git a/Engine/Plugins/Enterprise/LidarPointCloud/Source/LidarPointCloudEditor/LidarPointCloudEditor.Build.cs b/Engine/Plugins/Enterprise/LidarPointCloud/Source/LidarPointCloudEditor/LidarPointCloudEditor.Build.cs
--- a/Engine/Plugins/Enterprise/LidarPointCloud/Source/LidarPointCloudEditor/LidarPointCloudEditor.Build.cs
+++ b/Engine/Plugins/Enterprise/LidarPointCloud/Source/LidarPointCloudEditor/LidarPointCloudEditor.Build.cs
@@ -33,7 +33,8 @@ namespace UnrealBuildTool.Rules
                 "AdvancedPreviewScene",
                 "EditorInteractiveToolsFramework",
                 "InteractiveToolsFramework",
-                "StatusBar"
+                "StatusBar",
+                "MeshDescription",
                 }
             );
         }
diff --git a/Engine/Plugins/Enterprise/LidarPointCloud/Source/LidarPointCloudEditor/Private/LidarPointCloudEdMode.cpp b/Engine/Plugins/Enterprise/LidarPointCloud/Source/LidarPointCloudEditor/Private/LidarPointCloudEdMode.cpp
--- a/Engine/Plugins/Enterprise/LidarPointCloud/Source/LidarPointCloudEditor/Private/LidarPointCloudEdMode.cpp
+++ b/Engine/Plugins/Enterprise/LidarPointCloud/Source/LidarPointCloudEditor/Private/LidarPointCloudEdMode.cpp
@@ -44,6 +44,7 @@ void ULidarEditorMode::Enter()
 	REGISTER_TOOL(Merge);
 	REGISTER_TOOL(Collision);
 	REGISTER_TOOL(Normals);
+	REGISTER_TOOL(Meshing);
 	REGISTER_TOOL(BoxSelection);
 	REGISTER_TOOL(PolygonalSelection);
 	REGISTER_TOOL(LassoSelection);
diff --git a/Engine/Plugins/Enterprise/LidarPointCloud/Source/LidarPointCloudEditor/Private/LidarPointCloudEditor.cpp b/Engine/Plugins/Enterprise/LidarPointCloud/Source/LidarPointCloudEditor/Private/LidarPointCloudEditor.cpp
--- a/Engine/Plugins/Enterprise/LidarPointCloud/Source/LidarPointCloudEditor/Private/LidarPointCloudEditor.cpp
+++ b/Engine/Plugins/Enterprise/LidarPointCloud/Source/LidarPointCloudEditor/Private/LidarPointCloudEditor.cpp
@@ -941,7 +941,7 @@ TArray<FAssetData> FLidarPointCloudEditor::SelectAssets(const FText& Title)
 	OpenAssetDialogConfig.DialogTitleOverride = Title;
 	OpenAssetDialogConfig.DefaultPath = "/Game";
 	OpenAssetDialogConfig.bAllowMultipleSelection = true;
-	OpenAssetDialogConfig.AssetClassNames.Emplace("/Script/LidarPointCloudRuntime.LidarPointCloud");
+	OpenAssetDialogConfig.AssetClassNames.Add(ULidarPointCloud::StaticClass()->GetClassPathName());
 
 	FContentBrowserModule& ContentBrowserModule = FModuleManager::LoadModuleChecked<FContentBrowserModule>("ContentBrowser");
 	return ContentBrowserModule.Get().CreateModalOpenAssetDialog(OpenAssetDialogConfig);
diff --git a/Engine/Plugins/Enterprise/LidarPointCloud/Source/LidarPointCloudEditor/Private/LidarPointCloudEditorCommands.cpp b/Engine/Plugins/Enterprise/LidarPointCloud/Source/LidarPointCloudEditor/Private/LidarPointCloudEditorCommands.cpp
--- a/Engine/Plugins/Enterprise/LidarPointCloud/Source/LidarPointCloudEditor/Private/LidarPointCloudEditorCommands.cpp
+++ b/Engine/Plugins/Enterprise/LidarPointCloud/Source/LidarPointCloudEditor/Private/LidarPointCloudEditorCommands.cpp
@@ -49,6 +49,8 @@ void FLidarPointCloudEditorCommands::RegisterCommands()
 	ManageCommands.Add(ToolkitCollision);
 	UI_COMMAND(ToolkitNormals, "Normals", "Calculates normals", EUserInterfaceActionType::ToggleButton, FInputChord());
 	ManageCommands.Add(ToolkitNormals);
+	UI_COMMAND(ToolkitMeshing, "Meshing", "Build Static Meshes", EUserInterfaceActionType::ToggleButton, FInputChord());
+	ManageCommands.Add(ToolkitMeshing);
 	
 	TArray<TSharedPtr<FUICommandInfo, ESPMode::ThreadSafe>>& EditCommands = Commands.FindOrAdd(LidarEditorPalletes::Edit);
 	EditCommands.Add(ToolkitSelect);
diff --git a/Engine/Plugins/Enterprise/LidarPointCloud/Source/LidarPointCloudEditor/Private/LidarPointCloudEditorCommands.h b/Engine/Plugins/Enterprise/LidarPointCloud/Source/LidarPointCloudEditor/Private/LidarPointCloudEditorCommands.h
--- a/Engine/Plugins/Enterprise/LidarPointCloud/Source/LidarPointCloudEditor/Private/LidarPointCloudEditorCommands.h
+++ b/Engine/Plugins/Enterprise/LidarPointCloud/Source/LidarPointCloudEditor/Private/LidarPointCloudEditorCommands.h
@@ -62,6 +62,7 @@ public:
 	TSharedPtr<FUICommandInfo> ToolkitMerge;
 	TSharedPtr<FUICommandInfo> ToolkitCollision;
 	TSharedPtr<FUICommandInfo> ToolkitNormals;
+	TSharedPtr<FUICommandInfo> ToolkitMeshing;
 	TSharedPtr<FUICommandInfo> ToolkitBoxSelection;
 	TSharedPtr<FUICommandInfo> ToolkitPolygonalSelection;
 	TSharedPtr<FUICommandInfo> ToolkitLassoSelection;
diff --git a/Engine/Plugins/Enterprise/LidarPointCloud/Source/LidarPointCloudEditor/Private/LidarPointCloudEditorHelper.cpp b/Engine/Plugins/Enterprise/LidarPointCloud/Source/LidarPointCloudEditor/Private/LidarPointCloudEditorHelper.cpp
--- a/Engine/Plugins/Enterprise/LidarPointCloud/Source/LidarPointCloudEditor/Private/LidarPointCloudEditorHelper.cpp
+++ b/Engine/Plugins/Enterprise/LidarPointCloud/Source/LidarPointCloudEditor/Private/LidarPointCloudEditorHelper.cpp
@@ -8,11 +8,12 @@
 #include "LidarPointCloud.h"
 #include "LidarPointCloudActor.h"
 #include "LidarPointCloudComponent.h"
+#include "MeshDescription.h"
 #include "Selection.h"
-
 #include "AssetRegistry/AssetRegistryModule.h"
-
+#include "Engine/StaticMeshActor.h"
 #include "Misc/ScopedSlowTask.h"
+#include "PhysicsEngine/BodySetup.h"
 
 #define LOCTEXT_NAMESPACE "LidarPointCloudEditorHelper"
 
@@ -26,7 +27,7 @@ namespace
 		FSaveAssetDialogConfig SaveAssetDialogConfig;
 		SaveAssetDialogConfig.DialogTitleOverride = LOCTEXT("SelectDestination", "Select Destination");
 		SaveAssetDialogConfig.DefaultPath = "/Game";
-		SaveAssetDialogConfig.AssetClassNames.Emplace("/Script/LidarPointCloudRuntime.LidarPointCloud");
+		SaveAssetDialogConfig.AssetClassNames.Add(ULidarPointCloud::StaticClass()->GetClassPathName());
 		SaveAssetDialogConfig.ExistingAssetPolicy = ESaveAssetDialogExistingAssetPolicy::AllowButWarn;
 
 		FContentBrowserModule& ContentBrowserModule = FModuleManager::LoadModuleChecked<FContentBrowserModule>("ContentBrowser");
@@ -107,6 +108,22 @@ namespace
 		return Actors;
 	}
 
+	int32 GetNumLidarActors()
+	{
+		int32 NumActors = 0;
+		
+		for (TObjectIterator<ALidarPointCloudActor> It; It; ++It)
+		{
+			const ALidarPointCloudActor* LidarActor = Cast<ALidarPointCloudActor>(*It);
+			if(IsValid(LidarActor))
+			{
+				++NumActors;
+			}
+		}
+
+		return NumActors;
+	}
+
 	UWorld* GetFirstWorld()
 	{
 		for (TObjectIterator<ALidarPointCloudActor> It; It; ++It)
@@ -138,16 +155,142 @@ namespace
 		return PointClouds;
 	}
 
-	ALidarPointCloudActor* SpawnActor()
+	template<typename T>
+	T* SpawnActor()
 	{
-		ALidarPointCloudActor* NewActor = nullptr;
+		T* NewActor = nullptr;
 		if(UWorld* World = GetFirstWorld())
 		{
-			NewActor = Cast<ALidarPointCloudActor>(World->SpawnActor(ALidarPointCloudActor::StaticClass()));
+			NewActor = Cast<T>(World->SpawnActor(T::StaticClass()));
 		}
 		return NewActor;
 	}
 	
+	UStaticMesh* CreateStaticMesh()
+	{
+		UStaticMesh* StaticMesh = nullptr;
+		
+		// Initialize SaveAssetDialog config
+		FSaveAssetDialogConfig SaveAssetDialogConfig;
+		SaveAssetDialogConfig.DialogTitleOverride = LOCTEXT("SelectDestination", "Select Destination");
+		SaveAssetDialogConfig.DefaultPath = "/Game";
+		SaveAssetDialogConfig.AssetClassNames.Add(UStaticMesh::StaticClass()->GetClassPathName());
+		SaveAssetDialogConfig.ExistingAssetPolicy = ESaveAssetDialogExistingAssetPolicy::AllowButWarn;
+		
+		const FString SaveObjectPath = FModuleManager::LoadModuleChecked<FContentBrowserModule>("ContentBrowser").Get().CreateModalSaveAssetDialog(SaveAssetDialogConfig);
+		if (!SaveObjectPath.IsEmpty())
+		{
+			// Attempt to load existing asset first
+			StaticMesh = FindObject<UStaticMesh>(nullptr, *SaveObjectPath);
+
+			// Proceed to creating a new asset, if needed
+			if (!StaticMesh)
+			{
+				const FString PackageName = FPackageName::ObjectPathToPackageName(SaveObjectPath);
+				const FString ObjectName = FPackageName::ObjectPathToObjectName(SaveObjectPath);
+
+				StaticMesh = NewObject<UStaticMesh>(CreatePackage(*PackageName), UStaticMesh::StaticClass(), FName(*ObjectName), EObjectFlags::RF_Public | EObjectFlags::RF_Standalone);
+
+				FAssetRegistryModule::AssetCreated(StaticMesh);
+			}
+		}
+
+		return StaticMesh;
+	}
+
+	void AssignMeshBuffersToMesh(LidarPointCloudMeshing::FMeshBuffers* MeshBuffers, UStaticMesh* StaticMesh)
+	{
+		if(StaticMesh && MeshBuffers)
+		{
+			constexpr bool bCPUAccess = true;
+
+			bool bNewMesh = true;
+			if (StaticMesh->GetRenderData())
+			{
+				bNewMesh = false;
+				StaticMesh->ReleaseResources();
+				StaticMesh->ReleaseResourcesFence.Wait();
+			}
+			
+			StaticMesh->SetRenderData(MakeUnique<FStaticMeshRenderData>());
+			FStaticMeshRenderData* RenderData = StaticMesh->GetRenderData();
+			RenderData->AllocateLODResources(1);
+			FStaticMeshLODResources& LODResources = RenderData->LODResources[0];
+
+			TArray<FStaticMeshBuildVertex> StaticMeshBuildVertices;
+			StaticMeshBuildVertices.SetNum(MeshBuffers->Vertices.Num());
+			FStaticMeshBuildVertex* StaticMeshBuildVertexPtr = StaticMeshBuildVertices.GetData();
+
+			for(const LidarPointCloudMeshing::FVertexData& Vertex : MeshBuffers->Vertices)
+			{
+				StaticMeshBuildVertexPtr->Position = Vertex.Position;
+				StaticMeshBuildVertexPtr->Color = Vertex.Color;
+				StaticMeshBuildVertexPtr->TangentZ = Vertex.Normal;
+				StaticMeshBuildVertexPtr++;
+			}
+
+			LODResources.bBuffersInlined = true;
+			LODResources.IndexBuffer.TrySetAllowCPUAccess(bCPUAccess);
+			LODResources.VertexBuffers.PositionVertexBuffer.Init(StaticMeshBuildVertices);
+			LODResources.VertexBuffers.StaticMeshVertexBuffer.Init(StaticMeshBuildVertices, 1, bCPUAccess);
+			LODResources.bHasColorVertexData = true;
+			LODResources.VertexBuffers.ColorVertexBuffer.Init(StaticMeshBuildVertices);
+			LODResources.bHasReversedIndices = false;
+			LODResources.bHasReversedDepthOnlyIndices = false;
+			
+			FStaticMeshSection& Section = LODResources.Sections.AddDefaulted_GetRef();
+			Section.FirstIndex = 0;
+			Section.NumTriangles = MeshBuffers->Indices.Num() / 3;
+			Section.MinVertexIndex = 0;
+			Section.MaxVertexIndex = MeshBuffers->Vertices.Num() - 1;			
+			Section.MaterialIndex = 0;
+			Section.bEnableCollision = true;
+			Section.bCastShadow = true;			
+			LODResources.IndexBuffer.SetIndices(MeshBuffers->Indices, EIndexBufferStride::Force32Bit);
+				
+			LODResources.bHasDepthOnlyIndices = true;
+			LODResources.DepthOnlyIndexBuffer.SetIndices(MeshBuffers->Indices, EIndexBufferStride::Force32Bit);
+			LODResources.DepthOnlyNumTriangles = MeshBuffers->Indices.Num() / 3;
+
+			StaticMesh->InitResources();
+
+			// Set up RenderData bounds and LOD data
+			RenderData->Bounds = MeshBuffers->Bounds;
+			StaticMesh->CalculateExtendedBounds();
+				
+			RenderData->ScreenSize[0].Default = 1.0f;
+
+			// Set up physics-related data
+			StaticMesh->CreateBodySetup();
+			check(StaticMesh->GetBodySetup());
+			StaticMesh->GetBodySetup()->InvalidatePhysicsData();
+
+			if (!bNewMesh)
+			{
+				for (FThreadSafeObjectIterator Iter(UStaticMeshComponent::StaticClass()); Iter; ++Iter)
+				{
+					UStaticMeshComponent* StaticMeshComponent = Cast<UStaticMeshComponent>(*Iter);
+					if (StaticMeshComponent->GetStaticMesh() == StaticMesh)
+					{
+						// it needs to recreate IF it already has been created
+						if (StaticMeshComponent->IsPhysicsStateCreated())
+						{
+							StaticMeshComponent->RecreatePhysicsState();
+						}
+					}
+				}
+			}
+			
+			static UMaterialInterface* Material = Cast<UMaterialInterface>(FSoftObjectPath(TEXT("/LidarPointCloud/Materials/M_MeshedCloud.M_MeshedCloud")).TryLoad());
+			StaticMesh->AddMaterial(Material);
+
+			StaticMesh->SetLightingGuid(FGuid::NewGuid()) ;
+			StaticMesh->ImportVersion = EImportStaticMeshVersion::LastVersion;
+			StaticMesh->PostEditChange();
+			StaticMesh->MarkPackageDirty();
+		}
+	}
+	
 	FSceneView* GetEditorView(FEditorViewportClient* ViewportClient)
 	{
 		FSceneViewFamilyContext ViewFamily(FSceneViewFamily::ConstructionValues(
@@ -190,12 +333,15 @@ namespace
 		Origins.Last(1) = Origins[0] + ViewDirection * 99999999.0f;
 
 		// Calculate plane normals
-		const bool bFlipNormals = FVector::DotProduct(Normals[0], (MeanCenter - Origins[0])) > 0;
 		for (int32 i = 0; i < Points.Num(); ++i)
 		{
 			Normals[i] = ((Origins[(i + 1) % Points.Num()] - Origins[i]).GetSafeNormal() ^ Directions[i]).GetSafeNormal();
-			
-			if (bFlipNormals)
+		}
+
+		// Flip normals?
+		if(FVector::DotProduct(Normals[0], (MeanCenter - Origins[0])) > 0)
+		{
+			for (int32 i = 0; i < Points.Num(); ++i)
 			{
 				Normals[i] = -Normals[i];
 			}
@@ -346,6 +492,102 @@ void FLidarPointCloudEditorHelper::RemoveCollisionForSelection()
 	});
 }
 
+void FLidarPointCloudEditorHelper::MeshSelected(bool bMeshByPoints, float CellSize, bool bMergeMeshes, bool bRetainTransform)
+{
+	const int32 NumSteps = (bMeshByPoints ? GetNumLidarActors() : GEditor->GetSelectedActorCount()) + bMergeMeshes;
+	
+	FScopedSlowTask ProgressDialog(NumSteps, LOCTEXT("Meshing", "Meshing Point Clouds..."));
+	ProgressDialog.MakeDialog();
+	
+	if(bMergeMeshes)
+	{
+		if(UStaticMesh* StaticMesh = CreateStaticMesh())
+		{
+			FScopeBenchmarkTimer Timer("MakeMesh");
+			LidarPointCloudMeshing::FMeshBuffers MeshBuffers;
+
+			TFunction<void(ALidarPointCloudActor*)> MergeFunc = [CellSize, &ProgressDialog, MeshBuffersPtr = &MeshBuffers, bMeshByPoints](ALidarPointCloudActor* Actor)
+			{
+				if(ULidarPointCloud* PointCloud = Actor->GetPointCloud())
+				{
+					if(bMeshByPoints)
+					{
+						PointCloud->BuildStaticMeshBuffersForSelection(CellSize, MeshBuffersPtr, Actor->GetActorTransform());
+					}
+					else
+					{
+						PointCloud->BuildStaticMeshBuffers(CellSize, MeshBuffersPtr, Actor->GetActorTransform());
+					}
+				}
+					
+				ProgressDialog.EnterProgressFrame(1.0f);
+			};
+
+			if(bMeshByPoints)
+			{
+				ProcessAll(MoveTemp(MergeFunc));
+			}
+			else
+			{
+				ProcessSelection(MoveTemp(MergeFunc));
+			}
+		
+			AssignMeshBuffersToMesh(&MeshBuffers, StaticMesh);
+			ProgressDialog.EnterProgressFrame(1.0f);
+					
+			if(AStaticMeshActor* MeshActor = SpawnActor<AStaticMeshActor>())
+			{
+				MeshActor->GetStaticMeshComponent()->SetStaticMesh(StaticMesh);
+			}
+		}
+	}
+	else
+	{
+		TFunction<void(ALidarPointCloudActor*)> MergeFunc = [CellSize, &ProgressDialog, bRetainTransform, bMeshByPoints](ALidarPointCloudActor* Actor)
+		{
+			if(ULidarPointCloud* PointCloud = Actor->GetPointCloud())
+			{
+				if(UStaticMesh* StaticMesh = CreateStaticMesh())
+				{
+					LidarPointCloudMeshing::FMeshBuffers MeshBuffers;
+					
+					if(bMeshByPoints)
+					{
+						PointCloud->BuildStaticMeshBuffersForSelection(CellSize, &MeshBuffers, bRetainTransform ? FTransform::Identity : Actor->GetActorTransform());
+					}
+					else
+					{
+						PointCloud->BuildStaticMeshBuffers(CellSize, &MeshBuffers, bRetainTransform ? FTransform::Identity : Actor->GetActorTransform());
+					}
+					
+					AssignMeshBuffersToMesh(&MeshBuffers, StaticMesh);
+					
+					if(AStaticMeshActor* MeshActor = SpawnActor<AStaticMeshActor>())
+					{
+						if(bRetainTransform)
+						{
+							MeshActor->SetActorTransform(Actor->GetActorTransform());
+						}
+					 	
+						MeshActor->GetStaticMeshComponent()->SetStaticMesh(StaticMesh);
+					}
+				}
+			}
+				
+			ProgressDialog.EnterProgressFrame(1.0f);
+		};
+		
+		if(bMeshByPoints)
+		{
+			ProcessAll(MoveTemp(MergeFunc));
+		}
+		else
+		{
+			ProcessSelection(MoveTemp(MergeFunc));
+		}
+	}
+}
+
 void FLidarPointCloudEditorHelper::CalculateNormalsForSelection()
 {
 	ProcessSelection([](ALidarPointCloudActor* Actor)
@@ -394,7 +636,7 @@ void FLidarPointCloudEditorHelper::DeleteHidden()
 void FLidarPointCloudEditorHelper::Extract()
 {
 	ULidarPointCloud* NewPointCloud = Extract_Internal();
-	if(ALidarPointCloudActor* Actor = SpawnActor())
+	if(ALidarPointCloudActor* Actor = SpawnActor<ALidarPointCloudActor>())
 	{
 		NewPointCloud->RestoreOriginalCoordinates();
 		Actor->SetPointCloud(NewPointCloud);
@@ -406,7 +648,7 @@ void FLidarPointCloudEditorHelper::Extract()
 void FLidarPointCloudEditorHelper::ExtractAsCopy()
 {
 	ULidarPointCloud* NewPointCloud = Extract_Internal();
-	if(ALidarPointCloudActor* Actor = SpawnActor())
+	if(ALidarPointCloudActor* Actor = SpawnActor<ALidarPointCloudActor>())
 	{
 		NewPointCloud->RestoreOriginalCoordinates();
 		Actor->SetPointCloud(NewPointCloud);
diff --git a/Engine/Plugins/Enterprise/LidarPointCloud/Source/LidarPointCloudEditor/Private/LidarPointCloudEditorHelper.h b/Engine/Plugins/Enterprise/LidarPointCloud/Source/LidarPointCloudEditor/Private/LidarPointCloudEditorHelper.h
--- a/Engine/Plugins/Enterprise/LidarPointCloud/Source/LidarPointCloudEditor/Private/LidarPointCloudEditorHelper.h
+++ b/Engine/Plugins/Enterprise/LidarPointCloud/Source/LidarPointCloudEditor/Private/LidarPointCloudEditorHelper.h
@@ -33,6 +33,8 @@ public:
 	static void BuildCollisionForSelection();
 	static void SetCollisionErrorForSelection(float Error);
 	static void RemoveCollisionForSelection();
+
+	static void MeshSelected(bool bMeshByPoints, float CellSize, bool bMergeMeshes, bool bRetainTransform);
 	
 	static void CalculateNormalsForSelection();
 	static void SetNormalsQualityForSelection(int32 Quality, float NoiseTolerance);
diff --git a/Engine/Plugins/Enterprise/LidarPointCloud/Source/LidarPointCloudEditor/Private/LidarPointCloudEditorTools.cpp b/Engine/Plugins/Enterprise/LidarPointCloud/Source/LidarPointCloudEditor/Private/LidarPointCloudEditorTools.cpp
--- a/Engine/Plugins/Enterprise/LidarPointCloud/Source/LidarPointCloudEditor/Private/LidarPointCloudEditorTools.cpp
+++ b/Engine/Plugins/Enterprise/LidarPointCloud/Source/LidarPointCloudEditor/Private/LidarPointCloudEditorTools.cpp
@@ -114,6 +114,11 @@ void ULidarToolActions_Collision::RemoveCollision()
  	FLidarPointCloudEditorHelper::RemoveCollisionForSelection();
  }
 
+void ULidarToolActions_Meshing::BuildStaticMesh()
+{
+	FLidarPointCloudEditorHelper::MeshSelected(false, MaxMeshingError, bMergeMeshes, !bMergeMeshes && bRetainTransform);
+}
+
 void ULidarToolActions_Normals::CalculateNormals()
 {
 	FLidarPointCloudEditorHelper::SetNormalsQualityForSelection(Quality, NoiseTolerance);
@@ -166,6 +171,11 @@ void ULidarToolActions_Selection::ClearSelection()
 	FLidarPointCloudEditorHelper::ClearSelection();
 }
 
+void ULidarToolActions_Selection::BuildStaticMesh()
+{
+	FLidarPointCloudEditorHelper::MeshSelected(true, MaxMeshingError, bMergeMeshes, !bMergeMeshes && bRetainTransform);
+}
+
 void ULidarEditorTool_SelectionBase::DrawHUD(FCanvas* Canvas, IToolsContextRenderAPI* RenderAPI)
 {
 	if(Clicks.Num() == 0)
diff --git a/Engine/Plugins/Enterprise/LidarPointCloud/Source/LidarPointCloudEditor/Private/LidarPointCloudEditorTools.h b/Engine/Plugins/Enterprise/LidarPointCloud/Source/LidarPointCloudEditor/Private/LidarPointCloudEditorTools.h
--- a/Engine/Plugins/Enterprise/LidarPointCloud/Source/LidarPointCloudEditor/Private/LidarPointCloudEditorTools.h
+++ b/Engine/Plugins/Enterprise/LidarPointCloud/Source/LidarPointCloudEditor/Private/LidarPointCloudEditorTools.h
@@ -185,6 +185,43 @@ public:
 	virtual UInteractiveTool* BuildTool(const FToolBuilderState& SceneState) const override { return NewObject<ULidarEditorTool_Collision>(SceneState.ToolManager); }
 };
 
+UCLASS()
+class ULidarToolActions_Meshing : public UInteractiveToolPropertySet
+{
+	GENERATED_BODY()
+
+public:
+	/** Max error around the meshed areas. Leave at 0 for max quality */
+	UPROPERTY(EditAnywhere, Category = Options, meta = (UIMin = "0", UIMax = "2000", ClampMin = "0"))
+	float MaxMeshingError = 0;
+	
+	UPROPERTY(EditAnywhere, Category = Options)
+	bool bMergeMeshes = true;
+
+	/** If not merging meshes, this will retain the transform of the original cloud */
+	UPROPERTY(EditAnywhere, Category = Options, meta=(EditCondition="!bMergeMeshes"))
+	bool bRetainTransform = true;
+	
+	UFUNCTION(CallInEditor, Category = Actions)
+	void BuildStaticMesh();
+};
+
+UCLASS()
+class ULidarEditorTool_Meshing : public ULidarEditorTool_Base
+{
+public:
+	GENERATED_BODY()
+	virtual TObjectPtr<UInteractiveToolPropertySet> CreateToolActions() override { return NewObject<ULidarToolActions_Meshing>(this); }
+};
+
+UCLASS()
+class ULidarEditorToolBuilder_Meshing : public ULidarEditorToolBuilder_Base
+{
+public:
+	GENERATED_BODY()
+	virtual UInteractiveTool* BuildTool(const FToolBuilderState& SceneState) const override { return NewObject<ULidarEditorTool_Meshing>(SceneState.ToolManager); }
+};
+
 UCLASS()
 class ULidarToolActions_Normals : public UInteractiveToolPropertySet
 {
@@ -255,6 +292,20 @@ public:
 
 	UFUNCTION(CallInEditor, Category = Normals)
 	void CalculateNormals();
+
+	UFUNCTION(CallInEditor, Category = Meshing)
+	void BuildStaticMesh();
+	
+	/** Max error around the meshed areas. Leave at 0 for max quality */
+	UPROPERTY(EditAnywhere, Category = Meshing, meta = (UIMin = "0", UIMax = "2000", ClampMin = "0"))
+	float MaxMeshingError = 0;
+
+	UPROPERTY(EditAnywhere, Category = Meshing)
+	bool bMergeMeshes = true;
+
+	/** If not merging meshes, this will retain the transform of the original cloud */
+	UPROPERTY(EditAnywhere, Category = Meshing, meta=(EditCondition="!bMergeMeshes"))
+	bool bRetainTransform = true;
 };
 
 UCLASS()
diff --git a/Engine/Plugins/Enterprise/LidarPointCloud/Source/LidarPointCloudRuntime/LidarPointCloudRuntime.Build.cs b/Engine/Plugins/Enterprise/LidarPointCloud/Source/LidarPointCloudRuntime/LidarPointCloudRuntime.Build.cs
--- a/Engine/Plugins/Enterprise/LidarPointCloud/Source/LidarPointCloudRuntime/LidarPointCloudRuntime.Build.cs
+++ b/Engine/Plugins/Enterprise/LidarPointCloud/Source/LidarPointCloudRuntime/LidarPointCloudRuntime.Build.cs
@@ -21,7 +21,11 @@ namespace UnrealBuildTool.Rules
 					"RenderCore",
 					"Projects",
 					"RHI",
-					"InputCore"
+					"InputCore",
+					"GeometryCore",
+					"MeshConversion",
+					"MeshDescription",
+					"StaticMeshDescription"
 				}
             );
 
diff --git a/Engine/Plugins/Enterprise/LidarPointCloud/Source/LidarPointCloudRuntime/Private/Collision/LidarPointCloudCollision.cpp b/Engine/Plugins/Enterprise/LidarPointCloud/Source/LidarPointCloudRuntime/Private/Collision/LidarPointCloudCollision.cpp
deleted file mode 100644
diff --git a/Engine/Plugins/Enterprise/LidarPointCloud/Source/LidarPointCloudRuntime/Private/Collision/LidarPointCloudCollision.h b/Engine/Plugins/Enterprise/LidarPointCloud/Source/LidarPointCloudRuntime/Private/Collision/LidarPointCloudCollision.h
deleted file mode 100644
diff --git a/Engine/Plugins/Enterprise/LidarPointCloud/Source/LidarPointCloudRuntime/Private/LidarPointCloud.cpp b/Engine/Plugins/Enterprise/LidarPointCloud/Source/LidarPointCloudRuntime/Private/LidarPointCloud.cpp
--- a/Engine/Plugins/Enterprise/LidarPointCloud/Source/LidarPointCloudRuntime/Private/LidarPointCloud.cpp
+++ b/Engine/Plugins/Enterprise/LidarPointCloud/Source/LidarPointCloudRuntime/Private/LidarPointCloud.cpp
@@ -591,6 +591,11 @@ int64 ULidarPointCloud::NumSelectedPoints()
 	return Octree.NumSelectedPoints();
 }
 
+bool ULidarPointCloud::HasSelectedPoints()
+{
+	return Octree.HasSelectedPoints();
+}
+
 void ULidarPointCloud::GetSelectedPointsAsCopies(TArray64<FLidarPointCloudPoint>& SelectedPoints, FTransform Transform)
 {
 	Transform.AddToTranslation(LocationOffset);
@@ -606,6 +611,11 @@ void ULidarPointCloud::ClearSelection()
 {
 	Octree.ClearSelection();
 }
+
+void ULidarPointCloud::BuildStaticMeshBuffersForSelection(float CellSize, LidarPointCloudMeshing::FMeshBuffers* OutMeshBuffers, const FTransform& Transform)
+{
+	Octree.BuildStaticMeshBuffersForSelection(CellSize, OutMeshBuffers, Transform);
+}
 #endif
 
 void ULidarPointCloud::SetSourcePath(const FString& NewSourcePath)
@@ -716,6 +726,11 @@ void ULidarPointCloud::RemoveCollision()
 	bCollisionBuildInProgress = false;
 }
 
+void ULidarPointCloud::BuildStaticMeshBuffers(float CellSize, LidarPointCloudMeshing::FMeshBuffers* OutMeshBuffers, const FTransform& Transform)
+{
+	Octree.BuildStaticMeshBuffers(CellSize, OutMeshBuffers, Transform);
+}
+
 void ULidarPointCloud::SetLocationOffset(FVector Offset)
 {
 	LocationOffset = Offset;
diff --git a/Engine/Plugins/Enterprise/LidarPointCloud/Source/LidarPointCloudRuntime/Private/LidarPointCloudOctree.cpp b/Engine/Plugins/Enterprise/LidarPointCloud/Source/LidarPointCloudRuntime/Private/LidarPointCloudOctree.cpp
--- a/Engine/Plugins/Enterprise/LidarPointCloud/Source/LidarPointCloudRuntime/Private/LidarPointCloudOctree.cpp
+++ b/Engine/Plugins/Enterprise/LidarPointCloud/Source/LidarPointCloudRuntime/Private/LidarPointCloudOctree.cpp
@@ -4,7 +4,6 @@
 #include "LidarPointCloudOctreeMacros.h"
 #include "LidarPointCloud.h"
 #include "Meshing/LidarPointCloudMeshing.h"
-#include "Collision/LidarPointCloudCollision.h"
 #include "Misc/ScopeTryLock.h"
 #include "Containers/Queue.h"
 #include "Async/Async.h"
@@ -1003,7 +1002,7 @@ float FLidarPointCloudOctree::GetEstimatedPointSpacing() const
 
 void FLidarPointCloudOctree::BuildCollision(const float& Accuracy, const bool& bVisibleOnly)
 {
-	LidarPointCloudCollision::BuildCollisionMesh(this, Accuracy, bVisibleOnly, &CollisionMesh);
+	LidarPointCloudMeshing::BuildCollisionMesh(this, Accuracy, &CollisionMesh);
 }
 
 void FLidarPointCloudOctree::RemoveCollision()
@@ -1014,6 +1013,16 @@ void FLidarPointCloudOctree::RemoveCollision()
 	new (&CollisionMesh) FTriMeshCollisionData();
 }
 
+void FLidarPointCloudOctree::BuildStaticMeshBuffers(float CellSize, LidarPointCloudMeshing::FMeshBuffers* OutMeshBuffers, const FTransform& Transform)
+{
+	if(CellSize == 0)
+	{
+		CellSize = GetEstimatedPointSpacing() * 1.1f;
+	}
+	
+	LidarPointCloudMeshing::BuildStaticMeshBuffers(this, CellSize, false, OutMeshBuffers, Transform);
+}
+
 void FLidarPointCloudOctree::GetPoints(TArray<FLidarPointCloudPoint*>& SelectedPoints, int64 StartIndex /*= 0*/, int64 Count /*= -1*/)
 {
 	GetPoints_Internal(SelectedPoints, StartIndex, Count);
@@ -1879,6 +1888,12 @@ int64 FLidarPointCloudOctree::NumSelectedPoints()
 	return Count;
 }
 
+bool FLidarPointCloudOctree::HasSelectedPoints()
+{
+	ITERATE_SELECTED({ return true; }, {});
+	return false;
+}
+
 void FLidarPointCloudOctree::GetSelectedPointsAsCopies(TArray64<FLidarPointCloudPoint>& SelectedPoints, const FTransform& Transform)
 {
 	const FTransform3f Transform3F = (FTransform3f)Transform;
@@ -1889,6 +1904,19 @@ void FLidarPointCloudOctree::GetSelectedPointsAsCopies(TArray64<FLidarPointCloud
 	},{});
 }
 
+void FLidarPointCloudOctree::GetSelectedPointsInBox(TArray64<const FLidarPointCloudPoint*>& SelectedPoints, const FBox& Box) const
+{
+	constexpr bool bVisibleOnly = false;
+	
+	SelectedPoints.Reset(); 
+	PROCESS_IN_BOX({
+		if(Point->bSelected)
+		{
+			SelectedPoints.Add(Point);
+		}
+	});
+}
+
 void FLidarPointCloudOctree::CalculateNormalsForSelection(FThreadSafeBool* bCancelled, int32 Quality, float Tolerance)
 {
 	TArray64<FLidarPointCloudPoint*> SelectedPoints;
@@ -1919,6 +1947,16 @@ void FLidarPointCloudOctree::ClearSelection()
 		}
 	});
 }
+
+void FLidarPointCloudOctree::BuildStaticMeshBuffersForSelection(float CellSize, LidarPointCloudMeshing::FMeshBuffers* OutMeshBuffers, const FTransform& Transform)
+{
+	if(CellSize == 0)
+	{
+		CellSize = GetEstimatedPointSpacing() * 1.1f;
+	}
+	
+	LidarPointCloudMeshing::BuildStaticMeshBuffers(this, CellSize, true, OutMeshBuffers, Transform);
+}
 #endif // WITH_EDITOR
 
 void FLidarPointCloudOctree::MarkRenderDataInFrustumDirty(const FConvexVolume& Frustum)
diff --git a/Engine/Plugins/Enterprise/LidarPointCloud/Source/LidarPointCloudRuntime/Private/Meshing/LidarPointCloudMeshing.cpp b/Engine/Plugins/Enterprise/LidarPointCloud/Source/LidarPointCloudRuntime/Private/Meshing/LidarPointCloudMeshing.cpp
--- a/Engine/Plugins/Enterprise/LidarPointCloud/Source/LidarPointCloudRuntime/Private/Meshing/LidarPointCloudMeshing.cpp
+++ b/Engine/Plugins/Enterprise/LidarPointCloud/Source/LidarPointCloudRuntime/Private/Meshing/LidarPointCloudMeshing.cpp
@@ -1,14 +1,728 @@
 // Copyright Epic Games, Inc. All Rights Reserved.
  
 #include "Meshing/LidarPointCloudMeshing.h"
+#include "LidarPointCloud.h"
 #include "LidarPointCloudOctree.h"
 #include "LidarPointCloudOctreeMacros.h"
 #include "Async/Async.h"
 #include "Async/Future.h"
 
-FORCEINLINE uint64 CalculateNumPermutations(int32 NumPoints)
+namespace
 {
-	return NumPoints > 4000000 ? UINT64_MAX : NumPoints * (NumPoints - 1) * (NumPoints - 2) / 6;
+	FORCEINLINE uint64 CalculateNumPermutations(int32 NumPoints)
+	{
+		return NumPoints > 4000000 ? UINT64_MAX : NumPoints * (NumPoints - 1) * (NumPoints - 2) / 6;
+	}
+
+	namespace MarchingCubes
+	{
+		// Stores triangle indices
+		constexpr int8 TriangleTable[256][16] =
+		{
+			{-1, -1, -1, -1, -1, -1, -1, -1, -1, -1, -1, -1, -1, -1, -1, -1},
+			{0, 8, 3, -1, -1, -1, -1, -1, -1, -1, -1, -1, -1, -1, -1, -1},
+			{0, 1, 9, -1, -1, -1, -1, -1, -1, -1, -1, -1, -1, -1, -1, -1},
+			{1, 8, 3, 9, 8, 1, -1, -1, -1, -1, -1, -1, -1, -1, -1, -1},
+			{1, 2, 10, -1, -1, -1, -1, -1, -1, -1, -1, -1, -1, -1, -1, -1},
+			{0, 8, 3, 1, 2, 10, -1, -1, -1, -1, -1, -1, -1, -1, -1, -1},
+			{9, 2, 10, 0, 2, 9, -1, -1, -1, -1, -1, -1, -1, -1, -1, -1},
+			{2, 8, 3, 2, 10, 8, 10, 9, 8, -1, -1, -1, -1, -1, -1, -1},
+			{3, 11, 2, -1, -1, -1, -1, -1, -1, -1, -1, -1, -1, -1, -1, -1},
+			{0, 11, 2, 8, 11, 0, -1, -1, -1, -1, -1, -1, -1, -1, -1, -1},
+			{1, 9, 0, 2, 3, 11, -1, -1, -1, -1, -1, -1, -1, -1, -1, -1},
+			{1, 11, 2, 1, 9, 11, 9, 8, 11, -1, -1, -1, -1, -1, -1, -1},
+			{3, 10, 1, 11, 10, 3, -1, -1, -1, -1, -1, -1, -1, -1, -1, -1},
+			{0, 10, 1, 0, 8, 10, 8, 11, 10, -1, -1, -1, -1, -1, -1, -1},
+			{3, 9, 0, 3, 11, 9, 11, 10, 9, -1, -1, -1, -1, -1, -1, -1},
+			{9, 8, 10, 10, 8, 11, -1, -1, -1, -1, -1, -1, -1, -1, -1, -1},
+			{4, 7, 8, -1, -1, -1, -1, -1, -1, -1, -1, -1, -1, -1, -1, -1},
+			{4, 3, 0, 7, 3, 4, -1, -1, -1, -1, -1, -1, -1, -1, -1, -1},
+			{0, 1, 9, 8, 4, 7, -1, -1, -1, -1, -1, -1, -1, -1, -1, -1},
+			{4, 1, 9, 4, 7, 1, 7, 3, 1, -1, -1, -1, -1, -1, -1, -1},
+			{1, 2, 10, 8, 4, 7, -1, -1, -1, -1, -1, -1, -1, -1, -1, -1},
+			{3, 4, 7, 3, 0, 4, 1, 2, 10, -1, -1, -1, -1, -1, -1, -1},
+			{9, 2, 10, 9, 0, 2, 8, 4, 7, -1, -1, -1, -1, -1, -1, -1},
+			{2, 10, 9, 2, 9, 7, 2, 7, 3, 7, 9, 4, -1, -1, -1, -1},
+			{8, 4, 7, 3, 11, 2, -1, -1, -1, -1, -1, -1, -1, -1, -1, -1},
+			{11, 4, 7, 11, 2, 4, 2, 0, 4, -1, -1, -1, -1, -1, -1, -1},
+			{9, 0, 1, 8, 4, 7, 2, 3, 11, -1, -1, -1, -1, -1, -1, -1},
+			{4, 7, 11, 9, 4, 11, 9, 11, 2, 9, 2, 1, -1, -1, -1, -1},
+			{3, 10, 1, 3, 11, 10, 7, 8, 4, -1, -1, -1, -1, -1, -1, -1},
+			{1, 11, 10, 1, 4, 11, 1, 0, 4, 7, 11, 4, -1, -1, -1, -1},
+			{4, 7, 8, 9, 0, 11, 9, 11, 10, 11, 0, 3, -1, -1, -1, -1},
+			{4, 7, 11, 4, 11, 9, 9, 11, 10, -1, -1, -1, -1, -1, -1, -1},
+			{9, 5, 4, -1, -1, -1, -1, -1, -1, -1, -1, -1, -1, -1, -1, -1},
+			{9, 5, 4, 0, 8, 3, -1, -1, -1, -1, -1, -1, -1, -1, -1, -1},
+			{0, 5, 4, 1, 5, 0, -1, -1, -1, -1, -1, -1, -1, -1, -1, -1},
+			{8, 5, 4, 8, 3, 5, 3, 1, 5, -1, -1, -1, -1, -1, -1, -1},
+			{1, 2, 10, 9, 5, 4, -1, -1, -1, -1, -1, -1, -1, -1, -1, -1},
+			{3, 0, 8, 1, 2, 10, 4, 9, 5, -1, -1, -1, -1, -1, -1, -1},
+			{5, 2, 10, 5, 4, 2, 4, 0, 2, -1, -1, -1, -1, -1, -1, -1},
+			{2, 10, 5, 3, 2, 5, 3, 5, 4, 3, 4, 8, -1, -1, -1, -1},
+			{9, 5, 4, 2, 3, 11, -1, -1, -1, -1, -1, -1, -1, -1, -1, -1},
+			{0, 11, 2, 0, 8, 11, 4, 9, 5, -1, -1, -1, -1, -1, -1, -1},
+			{0, 5, 4, 0, 1, 5, 2, 3, 11, -1, -1, -1, -1, -1, -1, -1},
+			{2, 1, 5, 2, 5, 8, 2, 8, 11, 4, 8, 5, -1, -1, -1, -1},
+			{10, 3, 11, 10, 1, 3, 9, 5, 4, -1, -1, -1, -1, -1, -1, -1},
+			{4, 9, 5, 0, 8, 1, 8, 10, 1, 8, 11, 10, -1, -1, -1, -1},
+			{5, 4, 0, 5, 0, 11, 5, 11, 10, 11, 0, 3, -1, -1, -1, -1},
+			{5, 4, 8, 5, 8, 10, 10, 8, 11, -1, -1, -1, -1, -1, -1, -1},
+			{9, 7, 8, 5, 7, 9, -1, -1, -1, -1, -1, -1, -1, -1, -1, -1},
+			{9, 3, 0, 9, 5, 3, 5, 7, 3, -1, -1, -1, -1, -1, -1, -1},
+			{0, 7, 8, 0, 1, 7, 1, 5, 7, -1, -1, -1, -1, -1, -1, -1},
+			{1, 5, 3, 3, 5, 7, -1, -1, -1, -1, -1, -1, -1, -1, -1, -1},
+			{9, 7, 8, 9, 5, 7, 10, 1, 2, -1, -1, -1, -1, -1, -1, -1},
+			{10, 1, 2, 9, 5, 0, 5, 3, 0, 5, 7, 3, -1, -1, -1, -1},
+			{8, 0, 2, 8, 2, 5, 8, 5, 7, 10, 5, 2, -1, -1, -1, -1},
+			{2, 10, 5, 2, 5, 3, 3, 5, 7, -1, -1, -1, -1, -1, -1, -1},
+			{7, 9, 5, 7, 8, 9, 3, 11, 2, -1, -1, -1, -1, -1, -1, -1},
+			{9, 5, 7, 9, 7, 2, 9, 2, 0, 2, 7, 11, -1, -1, -1, -1},
+			{2, 3, 11, 0, 1, 8, 1, 7, 8, 1, 5, 7, -1, -1, -1, -1},
+			{11, 2, 1, 11, 1, 7, 7, 1, 5, -1, -1, -1, -1, -1, -1, -1},
+			{9, 5, 8, 8, 5, 7, 10, 1, 3, 10, 3, 11, -1, -1, -1, -1},
+			{5, 7, 0, 5, 0, 9, 7, 11, 0, 1, 0, 10, 11, 10, 0, -1},
+			{11, 10, 0, 11, 0, 3, 10, 5, 0, 8, 0, 7, 5, 7, 0, -1},
+			{11, 10, 5, 7, 11, 5, -1, -1, -1, -1, -1, -1, -1, -1, -1, -1},
+			{10, 6, 5, -1, -1, -1, -1, -1, -1, -1, -1, -1, -1, -1, -1, -1},
+			{0, 8, 3, 5, 10, 6, -1, -1, -1, -1, -1, -1, -1, -1, -1, -1},
+			{9, 0, 1, 5, 10, 6, -1, -1, -1, -1, -1, -1, -1, -1, -1, -1},
+			{1, 8, 3, 1, 9, 8, 5, 10, 6, -1, -1, -1, -1, -1, -1, -1},
+			{1, 6, 5, 2, 6, 1, -1, -1, -1, -1, -1, -1, -1, -1, -1, -1},
+			{1, 6, 5, 1, 2, 6, 3, 0, 8, -1, -1, -1, -1, -1, -1, -1},
+			{9, 6, 5, 9, 0, 6, 0, 2, 6, -1, -1, -1, -1, -1, -1, -1},
+			{5, 9, 8, 5, 8, 2, 5, 2, 6, 3, 2, 8, -1, -1, -1, -1},
+			{2, 3, 11, 10, 6, 5, -1, -1, -1, -1, -1, -1, -1, -1, -1, -1},
+			{11, 0, 8, 11, 2, 0, 10, 6, 5, -1, -1, -1, -1, -1, -1, -1},
+			{0, 1, 9, 2, 3, 11, 5, 10, 6, -1, -1, -1, -1, -1, -1, -1},
+			{5, 10, 6, 1, 9, 2, 9, 11, 2, 9, 8, 11, -1, -1, -1, -1},
+			{6, 3, 11, 6, 5, 3, 5, 1, 3, -1, -1, -1, -1, -1, -1, -1},
+			{0, 8, 11, 0, 11, 5, 0, 5, 1, 5, 11, 6, -1, -1, -1, -1},
+			{3, 11, 6, 0, 3, 6, 0, 6, 5, 0, 5, 9, -1, -1, -1, -1},
+			{6, 5, 9, 6, 9, 11, 11, 9, 8, -1, -1, -1, -1, -1, -1, -1},
+			{5, 10, 6, 4, 7, 8, -1, -1, -1, -1, -1, -1, -1, -1, -1, -1},
+			{4, 3, 0, 4, 7, 3, 6, 5, 10, -1, -1, -1, -1, -1, -1, -1},
+			{1, 9, 0, 5, 10, 6, 8, 4, 7, -1, -1, -1, -1, -1, -1, -1},
+			{10, 6, 5, 1, 9, 7, 1, 7, 3, 7, 9, 4, -1, -1, -1, -1},
+			{6, 1, 2, 6, 5, 1, 4, 7, 8, -1, -1, -1, -1, -1, -1, -1},
+			{1, 2, 5, 5, 2, 6, 3, 0, 4, 3, 4, 7, -1, -1, -1, -1},
+			{8, 4, 7, 9, 0, 5, 0, 6, 5, 0, 2, 6, -1, -1, -1, -1},
+			{7, 3, 9, 7, 9, 4, 3, 2, 9, 5, 9, 6, 2, 6, 9, -1},
+			{3, 11, 2, 7, 8, 4, 10, 6, 5, -1, -1, -1, -1, -1, -1, -1},
+			{5, 10, 6, 4, 7, 2, 4, 2, 0, 2, 7, 11, -1, -1, -1, -1},
+			{0, 1, 9, 4, 7, 8, 2, 3, 11, 5, 10, 6, -1, -1, -1, -1},
+			{9, 2, 1, 9, 11, 2, 9, 4, 11, 7, 11, 4, 5, 10, 6, -1},
+			{8, 4, 7, 3, 11, 5, 3, 5, 1, 5, 11, 6, -1, -1, -1, -1},
+			{5, 1, 11, 5, 11, 6, 1, 0, 11, 7, 11, 4, 0, 4, 11, -1},
+			{0, 5, 9, 0, 6, 5, 0, 3, 6, 11, 6, 3, 8, 4, 7, -1},
+			{6, 5, 9, 6, 9, 11, 4, 7, 9, 7, 11, 9, -1, -1, -1, -1},
+			{10, 4, 9, 6, 4, 10, -1, -1, -1, -1, -1, -1, -1, -1, -1, -1},
+			{4, 10, 6, 4, 9, 10, 0, 8, 3, -1, -1, -1, -1, -1, -1, -1},
+			{10, 0, 1, 10, 6, 0, 6, 4, 0, -1, -1, -1, -1, -1, -1, -1},
+			{8, 3, 1, 8, 1, 6, 8, 6, 4, 6, 1, 10, -1, -1, -1, -1},
+			{1, 4, 9, 1, 2, 4, 2, 6, 4, -1, -1, -1, -1, -1, -1, -1},
+			{3, 0, 8, 1, 2, 9, 2, 4, 9, 2, 6, 4, -1, -1, -1, -1},
+			{0, 2, 4, 4, 2, 6, -1, -1, -1, -1, -1, -1, -1, -1, -1, -1},
+			{8, 3, 2, 8, 2, 4, 4, 2, 6, -1, -1, -1, -1, -1, -1, -1},
+			{10, 4, 9, 10, 6, 4, 11, 2, 3, -1, -1, -1, -1, -1, -1, -1},
+			{0, 8, 2, 2, 8, 11, 4, 9, 10, 4, 10, 6, -1, -1, -1, -1},
+			{3, 11, 2, 0, 1, 6, 0, 6, 4, 6, 1, 10, -1, -1, -1, -1},
+			{6, 4, 1, 6, 1, 10, 4, 8, 1, 2, 1, 11, 8, 11, 1, -1},
+			{9, 6, 4, 9, 3, 6, 9, 1, 3, 11, 6, 3, -1, -1, -1, -1},
+			{8, 11, 1, 8, 1, 0, 11, 6, 1, 9, 1, 4, 6, 4, 1, -1},
+			{3, 11, 6, 3, 6, 0, 0, 6, 4, -1, -1, -1, -1, -1, -1, -1},
+			{6, 4, 8, 11, 6, 8, -1, -1, -1, -1, -1, -1, -1, -1, -1, -1},
+			{7, 10, 6, 7, 8, 10, 8, 9, 10, -1, -1, -1, -1, -1, -1, -1},
+			{0, 7, 3, 0, 10, 7, 0, 9, 10, 6, 7, 10, -1, -1, -1, -1},
+			{10, 6, 7, 1, 10, 7, 1, 7, 8, 1, 8, 0, -1, -1, -1, -1},
+			{10, 6, 7, 10, 7, 1, 1, 7, 3, -1, -1, -1, -1, -1, -1, -1},
+			{1, 2, 6, 1, 6, 8, 1, 8, 9, 8, 6, 7, -1, -1, -1, -1},
+			{2, 6, 9, 2, 9, 1, 6, 7, 9, 0, 9, 3, 7, 3, 9, -1},
+			{7, 8, 0, 7, 0, 6, 6, 0, 2, -1, -1, -1, -1, -1, -1, -1},
+			{7, 3, 2, 6, 7, 2, -1, -1, -1, -1, -1, -1, -1, -1, -1, -1},
+			{2, 3, 11, 10, 6, 8, 10, 8, 9, 8, 6, 7, -1, -1, -1, -1},
+			{2, 0, 7, 2, 7, 11, 0, 9, 7, 6, 7, 10, 9, 10, 7, -1},
+			{1, 8, 0, 1, 7, 8, 1, 10, 7, 6, 7, 10, 2, 3, 11, -1},
+			{11, 2, 1, 11, 1, 7, 10, 6, 1, 6, 7, 1, -1, -1, -1, -1},
+			{8, 9, 6, 8, 6, 7, 9, 1, 6, 11, 6, 3, 1, 3, 6, -1},
+			{0, 9, 1, 11, 6, 7, -1, -1, -1, -1, -1, -1, -1, -1, -1, -1},
+			{7, 8, 0, 7, 0, 6, 3, 11, 0, 11, 6, 0, -1, -1, -1, -1},
+			{7, 11, 6, -1, -1, -1, -1, -1, -1, -1, -1, -1, -1, -1, -1, -1},
+			{7, 6, 11, -1, -1, -1, -1, -1, -1, -1, -1, -1, -1, -1, -1, -1},
+			{3, 0, 8, 11, 7, 6, -1, -1, -1, -1, -1, -1, -1, -1, -1, -1},
+			{0, 1, 9, 11, 7, 6, -1, -1, -1, -1, -1, -1, -1, -1, -1, -1},
+			{8, 1, 9, 8, 3, 1, 11, 7, 6, -1, -1, -1, -1, -1, -1, -1},
+			{10, 1, 2, 6, 11, 7, -1, -1, -1, -1, -1, -1, -1, -1, -1, -1},
+			{1, 2, 10, 3, 0, 8, 6, 11, 7, -1, -1, -1, -1, -1, -1, -1},
+			{2, 9, 0, 2, 10, 9, 6, 11, 7, -1, -1, -1, -1, -1, -1, -1},
+			{6, 11, 7, 2, 10, 3, 10, 8, 3, 10, 9, 8, -1, -1, -1, -1},
+			{7, 2, 3, 6, 2, 7, -1, -1, -1, -1, -1, -1, -1, -1, -1, -1},
+			{7, 0, 8, 7, 6, 0, 6, 2, 0, -1, -1, -1, -1, -1, -1, -1},
+			{2, 7, 6, 2, 3, 7, 0, 1, 9, -1, -1, -1, -1, -1, -1, -1},
+			{1, 6, 2, 1, 8, 6, 1, 9, 8, 8, 7, 6, -1, -1, -1, -1},
+			{10, 7, 6, 10, 1, 7, 1, 3, 7, -1, -1, -1, -1, -1, -1, -1},
+			{10, 7, 6, 1, 7, 10, 1, 8, 7, 1, 0, 8, -1, -1, -1, -1},
+			{0, 3, 7, 0, 7, 10, 0, 10, 9, 6, 10, 7, -1, -1, -1, -1},
+			{7, 6, 10, 7, 10, 8, 8, 10, 9, -1, -1, -1, -1, -1, -1, -1},
+			{6, 8, 4, 11, 8, 6, -1, -1, -1, -1, -1, -1, -1, -1, -1, -1},
+			{3, 6, 11, 3, 0, 6, 0, 4, 6, -1, -1, -1, -1, -1, -1, -1},
+			{8, 6, 11, 8, 4, 6, 9, 0, 1, -1, -1, -1, -1, -1, -1, -1},
+			{9, 4, 6, 9, 6, 3, 9, 3, 1, 11, 3, 6, -1, -1, -1, -1},
+			{6, 8, 4, 6, 11, 8, 2, 10, 1, -1, -1, -1, -1, -1, -1, -1},
+			{1, 2, 10, 3, 0, 11, 0, 6, 11, 0, 4, 6, -1, -1, -1, -1},
+			{4, 11, 8, 4, 6, 11, 0, 2, 9, 2, 10, 9, -1, -1, -1, -1},
+			{10, 9, 3, 10, 3, 2, 9, 4, 3, 11, 3, 6, 4, 6, 3, -1},
+			{8, 2, 3, 8, 4, 2, 4, 6, 2, -1, -1, -1, -1, -1, -1, -1},
+			{0, 4, 2, 4, 6, 2, -1, -1, -1, -1, -1, -1, -1, -1, -1, -1},
+			{1, 9, 0, 2, 3, 4, 2, 4, 6, 4, 3, 8, -1, -1, -1, -1},
+			{1, 9, 4, 1, 4, 2, 2, 4, 6, -1, -1, -1, -1, -1, -1, -1},
+			{8, 1, 3, 8, 6, 1, 8, 4, 6, 6, 10, 1, -1, -1, -1, -1},
+			{10, 1, 0, 10, 0, 6, 6, 0, 4, -1, -1, -1, -1, -1, -1, -1},
+			{4, 6, 3, 4, 3, 8, 6, 10, 3, 0, 3, 9, 10, 9, 3, -1},
+			{10, 9, 4, 6, 10, 4, -1, -1, -1, -1, -1, -1, -1, -1, -1, -1},
+			{4, 9, 5, 7, 6, 11, -1, -1, -1, -1, -1, -1, -1, -1, -1, -1},
+			{0, 8, 3, 4, 9, 5, 11, 7, 6, -1, -1, -1, -1, -1, -1, -1},
+			{5, 0, 1, 5, 4, 0, 7, 6, 11, -1, -1, -1, -1, -1, -1, -1},
+			{11, 7, 6, 8, 3, 4, 3, 5, 4, 3, 1, 5, -1, -1, -1, -1},
+			{9, 5, 4, 10, 1, 2, 7, 6, 11, -1, -1, -1, -1, -1, -1, -1},
+			{6, 11, 7, 1, 2, 10, 0, 8, 3, 4, 9, 5, -1, -1, -1, -1},
+			{7, 6, 11, 5, 4, 10, 4, 2, 10, 4, 0, 2, -1, -1, -1, -1},
+			{3, 4, 8, 3, 5, 4, 3, 2, 5, 10, 5, 2, 11, 7, 6, -1},
+			{7, 2, 3, 7, 6, 2, 5, 4, 9, -1, -1, -1, -1, -1, -1, -1},
+			{9, 5, 4, 0, 8, 6, 0, 6, 2, 6, 8, 7, -1, -1, -1, -1},
+			{3, 6, 2, 3, 7, 6, 1, 5, 0, 5, 4, 0, -1, -1, -1, -1},
+			{6, 2, 8, 6, 8, 7, 2, 1, 8, 4, 8, 5, 1, 5, 8, -1},
+			{9, 5, 4, 10, 1, 6, 1, 7, 6, 1, 3, 7, -1, -1, -1, -1},
+			{1, 6, 10, 1, 7, 6, 1, 0, 7, 8, 7, 0, 9, 5, 4, -1},
+			{4, 0, 10, 4, 10, 5, 0, 3, 10, 6, 10, 7, 3, 7, 10, -1},
+			{7, 6, 10, 7, 10, 8, 5, 4, 10, 4, 8, 10, -1, -1, -1, -1},
+			{6, 9, 5, 6, 11, 9, 11, 8, 9, -1, -1, -1, -1, -1, -1, -1},
+			{3, 6, 11, 0, 6, 3, 0, 5, 6, 0, 9, 5, -1, -1, -1, -1},
+			{0, 11, 8, 0, 5, 11, 0, 1, 5, 5, 6, 11, -1, -1, -1, -1},
+			{6, 11, 3, 6, 3, 5, 5, 3, 1, -1, -1, -1, -1, -1, -1, -1},
+			{1, 2, 10, 9, 5, 11, 9, 11, 8, 11, 5, 6, -1, -1, -1, -1},
+			{0, 11, 3, 0, 6, 11, 0, 9, 6, 5, 6, 9, 1, 2, 10, -1},
+			{11, 8, 5, 11, 5, 6, 8, 0, 5, 10, 5, 2, 0, 2, 5, -1},
+			{6, 11, 3, 6, 3, 5, 2, 10, 3, 10, 5, 3, -1, -1, -1, -1},
+			{5, 8, 9, 5, 2, 8, 5, 6, 2, 3, 8, 2, -1, -1, -1, -1},
+			{9, 5, 6, 9, 6, 0, 0, 6, 2, -1, -1, -1, -1, -1, -1, -1},
+			{1, 5, 8, 1, 8, 0, 5, 6, 8, 3, 8, 2, 6, 2, 8, -1},
+			{1, 5, 6, 2, 1, 6, -1, -1, -1, -1, -1, -1, -1, -1, -1, -1},
+			{1, 3, 6, 1, 6, 10, 3, 8, 6, 5, 6, 9, 8, 9, 6, -1},
+			{10, 1, 0, 10, 0, 6, 9, 5, 0, 5, 6, 0, -1, -1, -1, -1},
+			{0, 3, 8, 5, 6, 10, -1, -1, -1, -1, -1, -1, -1, -1, -1, -1},
+			{10, 5, 6, -1, -1, -1, -1, -1, -1, -1, -1, -1, -1, -1, -1, -1},
+			{11, 5, 10, 7, 5, 11, -1, -1, -1, -1, -1, -1, -1, -1, -1, -1},
+			{11, 5, 10, 11, 7, 5, 8, 3, 0, -1, -1, -1, -1, -1, -1, -1},
+			{5, 11, 7, 5, 10, 11, 1, 9, 0, -1, -1, -1, -1, -1, -1, -1},
+			{10, 7, 5, 10, 11, 7, 9, 8, 1, 8, 3, 1, -1, -1, -1, -1},
+			{11, 1, 2, 11, 7, 1, 7, 5, 1, -1, -1, -1, -1, -1, -1, -1},
+			{0, 8, 3, 1, 2, 7, 1, 7, 5, 7, 2, 11, -1, -1, -1, -1},
+			{9, 7, 5, 9, 2, 7, 9, 0, 2, 2, 11, 7, -1, -1, -1, -1},
+			{7, 5, 2, 7, 2, 11, 5, 9, 2, 3, 2, 8, 9, 8, 2, -1},
+			{2, 5, 10, 2, 3, 5, 3, 7, 5, -1, -1, -1, -1, -1, -1, -1},
+			{8, 2, 0, 8, 5, 2, 8, 7, 5, 10, 2, 5, -1, -1, -1, -1},
+			{9, 0, 1, 5, 10, 3, 5, 3, 7, 3, 10, 2, -1, -1, -1, -1},
+			{9, 8, 2, 9, 2, 1, 8, 7, 2, 10, 2, 5, 7, 5, 2, -1},
+			{1, 3, 5, 3, 7, 5, -1, -1, -1, -1, -1, -1, -1, -1, -1, -1},
+			{0, 8, 7, 0, 7, 1, 1, 7, 5, -1, -1, -1, -1, -1, -1, -1},
+			{9, 0, 3, 9, 3, 5, 5, 3, 7, -1, -1, -1, -1, -1, -1, -1},
+			{9, 8, 7, 5, 9, 7, -1, -1, -1, -1, -1, -1, -1, -1, -1, -1},
+			{5, 8, 4, 5, 10, 8, 10, 11, 8, -1, -1, -1, -1, -1, -1, -1},
+			{5, 0, 4, 5, 11, 0, 5, 10, 11, 11, 3, 0, -1, -1, -1, -1},
+			{0, 1, 9, 8, 4, 10, 8, 10, 11, 10, 4, 5, -1, -1, -1, -1},
+			{10, 11, 4, 10, 4, 5, 11, 3, 4, 9, 4, 1, 3, 1, 4, -1},
+			{2, 5, 1, 2, 8, 5, 2, 11, 8, 4, 5, 8, -1, -1, -1, -1},
+			{0, 4, 11, 0, 11, 3, 4, 5, 11, 2, 11, 1, 5, 1, 11, -1},
+			{0, 2, 5, 0, 5, 9, 2, 11, 5, 4, 5, 8, 11, 8, 5, -1},
+			{9, 4, 5, 2, 11, 3, -1, -1, -1, -1, -1, -1, -1, -1, -1, -1},
+			{2, 5, 10, 3, 5, 2, 3, 4, 5, 3, 8, 4, -1, -1, -1, -1},
+			{5, 10, 2, 5, 2, 4, 4, 2, 0, -1, -1, -1, -1, -1, -1, -1},
+			{3, 10, 2, 3, 5, 10, 3, 8, 5, 4, 5, 8, 0, 1, 9, -1},
+			{5, 10, 2, 5, 2, 4, 1, 9, 2, 9, 4, 2, -1, -1, -1, -1},
+			{8, 4, 5, 8, 5, 3, 3, 5, 1, -1, -1, -1, -1, -1, -1, -1},
+			{0, 4, 5, 1, 0, 5, -1, -1, -1, -1, -1, -1, -1, -1, -1, -1},
+			{8, 4, 5, 8, 5, 3, 9, 0, 5, 0, 3, 5, -1, -1, -1, -1},
+			{9, 4, 5, -1, -1, -1, -1, -1, -1, -1, -1, -1, -1, -1, -1, -1},
+			{4, 11, 7, 4, 9, 11, 9, 10, 11, -1, -1, -1, -1, -1, -1, -1},
+			{0, 8, 3, 4, 9, 7, 9, 11, 7, 9, 10, 11, -1, -1, -1, -1},
+			{1, 10, 11, 1, 11, 4, 1, 4, 0, 7, 4, 11, -1, -1, -1, -1},
+			{3, 1, 4, 3, 4, 8, 1, 10, 4, 7, 4, 11, 10, 11, 4, -1},
+			{4, 11, 7, 9, 11, 4, 9, 2, 11, 9, 1, 2, -1, -1, -1, -1},
+			{9, 7, 4, 9, 11, 7, 9, 1, 11, 2, 11, 1, 0, 8, 3, -1},
+			{11, 7, 4, 11, 4, 2, 2, 4, 0, -1, -1, -1, -1, -1, -1, -1},
+			{11, 7, 4, 11, 4, 2, 8, 3, 4, 3, 2, 4, -1, -1, -1, -1},
+			{2, 9, 10, 2, 7, 9, 2, 3, 7, 7, 4, 9, -1, -1, -1, -1},
+			{9, 10, 7, 9, 7, 4, 10, 2, 7, 8, 7, 0, 2, 0, 7, -1},
+			{3, 7, 10, 3, 10, 2, 7, 4, 10, 1, 10, 0, 4, 0, 10, -1},
+			{1, 10, 2, 8, 7, 4, -1, -1, -1, -1, -1, -1, -1, -1, -1, -1},
+			{4, 9, 1, 4, 1, 7, 7, 1, 3, -1, -1, -1, -1, -1, -1, -1},
+			{4, 9, 1, 4, 1, 7, 0, 8, 1, 8, 7, 1, -1, -1, -1, -1},
+			{4, 0, 3, 7, 4, 3, -1, -1, -1, -1, -1, -1, -1, -1, -1, -1},
+			{4, 8, 7, -1, -1, -1, -1, -1, -1, -1, -1, -1, -1, -1, -1, -1},
+			{9, 10, 8, 10, 11, 8, -1, -1, -1, -1, -1, -1, -1, -1, -1, -1},
+			{3, 0, 9, 3, 9, 11, 11, 9, 10, -1, -1, -1, -1, -1, -1, -1},
+			{0, 1, 10, 0, 10, 8, 8, 10, 11, -1, -1, -1, -1, -1, -1, -1},
+			{3, 1, 10, 11, 3, 10, -1, -1, -1, -1, -1, -1, -1, -1, -1, -1},
+			{1, 2, 11, 1, 11, 9, 9, 11, 8, -1, -1, -1, -1, -1, -1, -1},
+			{3, 0, 9, 3, 9, 11, 1, 2, 9, 2, 11, 9, -1, -1, -1, -1},
+			{0, 2, 11, 8, 0, 11, -1, -1, -1, -1, -1, -1, -1, -1, -1, -1},
+			{3, 2, 11, -1, -1, -1, -1, -1, -1, -1, -1, -1, -1, -1, -1, -1},
+			{2, 3, 8, 2, 8, 10, 10, 8, 9, -1, -1, -1, -1, -1, -1, -1},
+			{9, 10, 2, 0, 9, 2, -1, -1, -1, -1, -1, -1, -1, -1, -1, -1},
+			{2, 3, 8, 2, 8, 10, 0, 1, 8, 1, 10, 8, -1, -1, -1, -1},
+			{1, 10, 2, -1, -1, -1, -1, -1, -1, -1, -1, -1, -1, -1, -1, -1},
+			{1, 3, 8, 9, 1, 8, -1, -1, -1, -1, -1, -1, -1, -1, -1, -1},
+			{0, 9, 1, -1, -1, -1, -1, -1, -1, -1, -1, -1, -1, -1, -1, -1},
+			{0, 3, 8, -1, -1, -1, -1, -1, -1, -1, -1, -1, -1, -1, -1, -1},
+			{-1, -1, -1, -1, -1, -1, -1, -1, -1, -1, -1, -1, -1, -1, -1, -1}
+		};
+
+		const FVector3f VertexOffsetTable[12] =
+		{
+			FVector3f(0.5, 0, 0),
+			FVector3f(1, 0.5, 0),
+			FVector3f(0.5, 1, 0),
+			FVector3f(0, 0.5, 0),
+			FVector3f(0.5, 0, 1),
+			FVector3f(1, 0.5, 1),
+			FVector3f(0.5, 1, 1),
+			FVector3f(0, 0.5, 1),
+			FVector3f(0, 0, 0.5),
+			FVector3f(1, 0, 0.5),
+			FVector3f(1, 1, 0.5),
+			FVector3f(0, 1, 0.5)
+		};
+
+		// 256 possible marching cubes combinations
+		constexpr int32 NumCombinations = 256;
+
+		// Max 5 triangles x 3 vertices x 12 bytes each + 1 byte for NumVertices 
+		constexpr int32 MaxCombinationSize = 181;
+
+		/** Calculate pre-computed data, ready for processing. */
+		TArray<uint8> BuildScaledVertexTable(float CellSize)
+		{
+			TArray<uint8> VertexData;
+			VertexData.AddUninitialized(NumCombinations * MaxCombinationSize);
+			
+			TArray<FVector3f> TempArray;
+			uint8* VertexDataPtr = VertexData.GetData();
+			for (int32 i = 0; i < NumCombinations; ++i)
+			{
+				const int8* VertexList = TriangleTable[i];
+				TempArray.Reset();
+
+				for (int32 t = 0; t < 16; t += 3)
+				{
+					if (VertexList[t] == -1)
+					{
+						break;
+					}
+
+					TempArray.Add(VertexOffsetTable[VertexList[t + 2]] * CellSize);
+					TempArray.Add(VertexOffsetTable[VertexList[t + 1]] * CellSize);
+					TempArray.Add(VertexOffsetTable[VertexList[t]] * CellSize);
+				}
+
+				*VertexDataPtr++ = TempArray.Num();
+				FMemory::Memcpy(VertexDataPtr, TempArray.GetData(), TempArray.Num() * 12);
+				VertexDataPtr += MaxCombinationSize - 1;
+			}
+
+			return VertexData;
+		}
+
+		/** Constructs the normals table for vertices */
+		TArray<TArray<FVector3f>>* GetNormalsTable()
+		{
+			static TArray<TArray<FVector3f>> NormalsTable;
+
+			if(NormalsTable.Num() == 0)
+			{
+				NormalsTable.AddDefaulted(NumCombinations);
+				for(int32 i = 0; i < NumCombinations; ++i)
+				{
+					TArray<FVector3f>& Normals = NormalsTable[i];
+					Normals.Reserve(15);
+					
+					for(int32 j = 0; j < 15; j += 3)
+					{
+						const int8 A = TriangleTable[i][j];
+						const int8 B = TriangleTable[i][j+1];
+						const int8 C = TriangleTable[i][j+2];
+
+						if(A > -1)
+						{
+							const FVector3f P1 = VertexOffsetTable[A];
+							const FVector3f P2 = VertexOffsetTable[B];
+							const FVector3f P3 = VertexOffsetTable[C];
+
+							const FVector3f V1 = (P3 - P2).GetSafeNormal();
+							const FVector3f V2 = (P1 - P2).GetSafeNormal();
+
+							const FVector3f Normal = FVector3f::CrossProduct(V1, V2).GetSafeNormal();
+							Normals.Append({Normal, Normal, Normal});
+						}
+						else
+						{
+							static TArray<FVector3f> EmptyNormal({ FVector3f::ZeroVector, FVector3f::ZeroVector, FVector3f::ZeroVector });
+							Normals.Append(EmptyNormal);
+						}
+					}
+				}
+			}
+
+			return &NormalsTable;
+		}
+
+		/** Calculates the color lookup table for vertices */
+		void CalculateColorLookupTable(int32 GridDimension, TArray<TArray<uint32>>& OutColorTable)
+		{
+			static TArray<TArray<FIntVector>> MasterColorOffsetTable;
+
+			if(MasterColorOffsetTable.Num() == 0)
+			{
+				TArray<TArray<FVector3f>>* NormalsTable = GetNormalsTable();
+
+				MasterColorOffsetTable.AddDefaulted(NumCombinations);
+				for(int32 i = 0; i < NumCombinations; ++i)
+				{
+					TArray<FIntVector>& ColorOffsets = MasterColorOffsetTable[i];
+					ColorOffsets.Reserve(15);
+					
+					for(int32 j = 0; j < 15; ++j)
+					{
+						const int8 Idx = TriangleTable[i][j];
+						if(Idx > -1)
+						{
+							const static FVector3f Corners[8] =
+							{
+								FVector3f(0, 0, 0),
+								FVector3f(1, 0, 0),
+								FVector3f(0, 1, 0),
+								FVector3f(1, 1, 0),
+								FVector3f(0, 0, 1),
+								FVector3f(1, 0, 1),
+								FVector3f(0, 1, 1),
+								FVector3f(1, 1, 1)
+							};
+							
+							const FVector3f P = VertexOffsetTable[Idx];
+							const FVector3f N = -(*NormalsTable)[i][j];
+
+							float MinDist = FLT_MAX;
+							int32 MinC = 0;
+
+							for(int32 c = 0; c < 8; c++)
+							{
+								const FVector3f V = Corners[c] - P;
+								if(FVector3f::DotProduct(N, V) > 0)
+								{
+									const float Dist = V.SizeSquared();
+								
+									if(Dist < MinDist)
+									{
+										MinDist = Dist;
+										MinC = c;
+									}
+								}
+							}
+
+							ColorOffsets.Emplace((FVector)Corners[MinC]);
+						}
+					}
+				}
+			}
+
+			const int32 GridDimensionSq = GridDimension * GridDimension;
+
+			OutColorTable.SetNum(NumCombinations);
+			for(int32 i = 0; i < NumCombinations; ++i)
+			{
+				TArray<FIntVector>& MasterOffsets = MasterColorOffsetTable[i];
+				TArray<uint32>& ColorOffsets = OutColorTable[i];
+				ColorOffsets.SetNumZeroed(MasterOffsets.Num());
+
+				for(int32 j = 0; j < ColorOffsets.Num(); ++j)
+				{
+					const FIntVector& MasterOffset = MasterOffsets[j];
+					ColorOffsets[j] = MasterOffset.Z * GridDimensionSq + MasterOffset.Y * GridDimension + MasterOffset.X;
+				}
+			}
+		}
+
+		/** Applies Marching Cubes algorithm to extract mesh information from the given voxelized grid */
+		void ProcessGrid(uint8* VoxelizedGrid, uint8* VertexData, FColor* ColorGrid, TArray<TArray<uint32>>& ColorLookupTable, int32 GridDimension, float CellSize, FVector3f VertexOffset,
+			TArray<FVector3f>& OutVertices, TArray<FVector3f>* OutNormals, TArray<FColor>* OutColors)
+		{
+			TArray<TArray<FVector3f>>* NormalsTable = GetNormalsTable();
+			
+			FVector3f Coords;
+			uint8* Vertices = nullptr;
+
+			const int32 GridDimensionSq = GridDimension * GridDimension;
+
+			uint8* Grid = VoxelizedGrid;
+			uint8* GridUp = Grid + GridDimensionSq;
+
+			for (int32 sZ = 0; sZ < GridDimension - 1; ++sZ)
+			{
+				Coords.Z = VertexOffset.Z + (sZ + 0.5f) * CellSize;
+
+				Grid = VoxelizedGrid + sZ * GridDimensionSq;
+				GridUp = Grid + GridDimensionSq;
+				int32 BaseIndex = 0;
+
+				for (int32 sY = 0; sY < GridDimension - 1; ++sY)
+				{
+					Coords.Y = VertexOffset.Y + (sY + 0.5f) * CellSize;
+
+					for (int32 sX = 0; sX < GridDimension - 1; ++sX, ++BaseIndex)
+					{
+						Coords.X = VertexOffset.X + (sX + 0.5f) * CellSize;
+
+						// Calculate the variation of the cube
+						uint8 CubeVariation = 0;
+						if (Grid[BaseIndex] == 0) { CubeVariation |= 1; }
+						if (Grid[BaseIndex + 1] == 0) { CubeVariation |= 2; }
+						if (Grid[BaseIndex + GridDimension + 1] == 0) { CubeVariation |= 4; }
+						if (Grid[BaseIndex + GridDimension] == 0) { CubeVariation |= 8; }
+						if (GridUp[BaseIndex] == 0) { CubeVariation |= 16; }
+						if (GridUp[BaseIndex + 1] == 0) { CubeVariation |= 32; }
+						if (GridUp[BaseIndex + GridDimension + 1] == 0) { CubeVariation |= 64; }
+						if (GridUp[BaseIndex + GridDimension] == 0) { CubeVariation |= 128; }
+
+						// No triangles to add
+						if (CubeVariation == 0 || CubeVariation == 255)
+						{
+							continue;
+						}
+
+						Vertices = &VertexData[CubeVariation * MaxCombinationSize];
+						const uint8 NumVertices = Vertices[0];
+
+						if (NumVertices == 0)
+						{
+							continue;
+						}
+
+						// Fill vertex positions
+						{
+							const int32 Offset = OutVertices.AddUninitialized(NumVertices);
+							for (FVector3f* Vertex = OutVertices.GetData() + Offset, *SourceData = (FVector3f*)(Vertices + 1), *DataEnd = Vertex + NumVertices; Vertex != DataEnd; ++Vertex, ++SourceData)
+							{
+								*Vertex = Coords + *SourceData;
+							}
+						}
+
+						// Fill normals if requested
+						if(OutNormals)
+						{
+							const int32 Offset = OutNormals->AddUninitialized(NumVertices);
+							FMemory::Memcpy(OutNormals->GetData() + Offset, (*NormalsTable)[CubeVariation].GetData(), NumVertices * sizeof(FVector3f));
+						}
+
+						// Fill colors if requested
+						if(OutColors)
+						{
+							const uint32 BaseColorIndex = sZ * GridDimensionSq + BaseIndex;
+							const uint32* LookupTablePtr = ColorLookupTable[CubeVariation].GetData();
+							const int32 Offset = OutColors->AddUninitialized(NumVertices);
+							FColor* ColorsPtr = OutColors->GetData() + Offset;
+							for(int32 v = 0; v < NumVertices; ++v)
+							{
+								*ColorsPtr++ = ColorGrid[BaseColorIndex + *LookupTablePtr++];
+							}
+						}
+					}
+
+					++BaseIndex;
+				}
+			}
+		}
+
+		void Run(FLidarPointCloudOctree* Octree, const float& CellSize, bool bUseSelection, TArray<FVector3f>& OutVertices, TArray<FVector3f>* OutNormals, TArray<FColor>* OutColors)
+		{
+			FScopeBenchmarkTimer Timer("Marching Cubes");
+			
+			// Expand original bounds to make sure the mesh is closed at the edges
+			const FBox Bounds = Octree->GetBounds().ExpandBy(FVector::OneVector * CellSize, FVector::OneVector * CellSize);
+			const FVector3f BoundsSize = (FVector3f)Bounds.GetSize();
+			const FVector3f LocationOffset = (FVector3f)Octree->GetOwner()->LocationOffset;
+
+			// Number of cells in each axis
+			const int32 BatchSize = GetDefault<ULidarPointCloudSettings>()->MeshingBatchSize;
+
+			// Determines the number of samples to perform
+			const FIntVector NumSamples((FVector)BoundsSize / ((BatchSize - 1) * CellSize) + 1);
+			const int32 TotalNumSamples = NumSamples.X * NumSamples.Y * NumSamples.Z;
+
+			// Precalculated for performance reasons
+			const float InversedCellSize = 1 / CellSize;
+			const FBox BaseSamplingBounds(Bounds.Min, Bounds.Min + BatchSize * CellSize);
+			TArray<TArray<uint32>> ColorLookupTable;
+			CalculateColorLookupTable(BatchSize, ColorLookupTable);
+
+			// Multithreading
+			const int32 MaxNumThreads = FMath::Min(FPlatformMisc::NumberOfCoresIncludingHyperthreads() - 1, TotalNumSamples);
+			TArray<TFuture<void>> ThreadResults;
+			FThreadSafeCounter SampleIndex = 0;
+
+			// Data storage
+			FCriticalSection VertexLock;
+			OutVertices.Empty();
+
+			TArray<uint8> VertexData = BuildScaledVertexTable(CellSize);
+			uint8* VertexDataPtr = VertexData.GetData();
+			
+			// Make sure to acquire data release lock before loading nodes
+			FScopeLock ReleaseLock(&Octree->DataReleaseLock);
+
+			Octree->LoadAllNodes(false);
+
+			const FLidarPointCloudOctree* OctreeConst = Octree;
+
+			// Fire threads
+			for (int32 t = 0; t < MaxNumThreads && SampleIndex.GetValue() < TotalNumSamples; t++)
+			{
+				ThreadResults.Add(Async(EAsyncExecution::TaskGraph, [&SampleIndex, bUseSelection, &OutVertices, &VertexLock, VertexDataPtr, BatchSize,
+					OutNormals, OutColors, &ColorLookupTable, TotalNumSamples, NumSamples, BaseSamplingBounds, CellSize, OctreeConst, InversedCellSize, LocationOffset]
+				{
+					// Local caching arrays to reduce number of syncs required
+					TArray<FVector3f> _Vertices;
+					TArray<FVector3f> _Normals;
+					TArray<FVector3f>* _NormalsPtr = OutNormals ? &_Normals : nullptr;
+					TArray<FColor> _Colors;
+					TArray<FColor>* _ColorsPtr = OutColors ? &_Colors : nullptr;
+					
+					const int32 NumCellsSq = BatchSize * BatchSize;
+					const int32 NumCellsCu = NumCellsSq * BatchSize;
+
+					uint8* VoxelizedGrid = new uint8[NumCellsCu];
+					FColor* ColorGrid = nullptr;
+					if(OutColors)
+					{
+						ColorGrid = new FColor[NumCellsCu];
+					}
+
+					const float OffsetMultiplier = (BatchSize - 1) * CellSize;
+
+					TArray64<const FLidarPointCloudPoint*> Selection;
+
+					int32 Index;
+					while ((Index = SampleIndex.Add(1)) < TotalNumSamples)
+					{
+						const int32 W = Index % (NumSamples.X * NumSamples.Y);
+						const FBox SamplingBounds = BaseSamplingBounds.ShiftBy(FVector(W % NumSamples.X, W / NumSamples.X, Index / (NumSamples.X * NumSamples.Y)) * OffsetMultiplier);
+
+						// Sample the data using the calculated bounds
+#if WITH_EDITOR
+						if(bUseSelection)
+						{
+							OctreeConst->GetSelectedPointsInBox(Selection, SamplingBounds);
+						}
+						else
+#endif
+						{
+							OctreeConst->GetPointsInBox(Selection, SamplingBounds, true);
+						}
+						
+						// Skip if no points in selection
+						if (Selection.Num() == 0)
+						{
+							continue;
+						}
+
+						// Clear existing grids
+						FMemory::Memzero(VoxelizedGrid, NumCellsCu);
+						if(ColorGrid)
+						{
+							FMemory::Memzero(ColorGrid, NumCellsCu * sizeof(FColor));
+						}
+
+						// Calculate voxelized grids
+						for (const FLidarPointCloudPoint** Point = Selection.GetData(), **DataEnd = Selection.GetData() + Selection.Num(); Point != DataEnd; ++Point)
+						{
+							// Calculate location relative to sampling bounds
+							const FVector3f Location = (*Point)->Location - (FVector3f)SamplingBounds.Min;
+
+							// Calculate grid coordinates
+							const FIntVector Grid(FMath::Min(BatchSize - 1, (int32)(Location.X * InversedCellSize)), FMath::Min(BatchSize - 1, (int32)(Location.Y * InversedCellSize)), FMath::Min(BatchSize - 1, (int32)(Location.Z * InversedCellSize)));
+
+							const int32 GridIndex = Grid.Z * NumCellsSq + Grid.Y * BatchSize + Grid.X;
+
+							VoxelizedGrid[GridIndex] = 1;
+
+							if(ColorGrid)
+							{
+								ColorGrid[GridIndex] = (*Point)->Color;
+							}
+						}
+
+						ProcessGrid(VoxelizedGrid, VertexDataPtr, ColorGrid, ColorLookupTable, BatchSize, CellSize, (FVector3f)SamplingBounds.Min + LocationOffset, _Vertices, _NormalsPtr, _ColorsPtr);
+					}
+				
+					// Clean up
+					delete[] VoxelizedGrid;
+					if(ColorGrid)
+					{
+						delete[] ColorGrid;
+					}
+
+					// Sync
+					FScopeLock Lock(&VertexLock);
+					OutVertices.Append(_Vertices);
+					if(OutNormals)
+					{
+						OutNormals->Append(_Normals);
+					}
+					if(OutColors)
+					{
+						OutColors->Append(_Colors);
+					}
+				}));
+			}
+
+			// Sync threads
+			for (const TFuture<void>& ThreadResult : ThreadResults)
+			{
+				ThreadResult.Get();
+			}
+
+			Octree->ReleaseAllNodes(false);
+		}
+	}
+}
+
+void LidarPointCloudMeshing::FMeshBuffers::Init(uint32 Capacity, bool bExpandIfNotEmpty, uint32& OutIndexOffset, uint32& OutVertexOffset)
+{
+	if((Indices.IsEmpty() && Vertices.IsEmpty()) || !bExpandIfNotEmpty)
+	{
+		OutIndexOffset = OutVertexOffset = 0;
+		Indices.SetNumUninitialized(Capacity);
+		Vertices.Reserve(Capacity);
+		Bounds = FBoxSphereBounds(EForceInit::ForceInit);
+	}
+	else
+	{
+		OutIndexOffset = Indices.AddUninitialized(Capacity);
+		OutVertexOffset = Vertices.Num();
+        Vertices.Reserve(Vertices.Num() + Capacity);
+	}
+}
+
+void LidarPointCloudMeshing::FMeshBuffers::ExpandBounds(const FBox& NewBounds)
+{
+	const FBoxSphereBounds ConvertedNewBounds(NewBounds);
+	Bounds = Bounds.BoxExtent == FVector::ZeroVector ? ConvertedNewBounds : Bounds + ConvertedNewBounds;
+}
+
+void LidarPointCloudMeshing::FMeshBuffers::NormalizeNormals()
+{
+	for (FVertexData *Vertex = Vertices.GetData(), *DataEnd = Vertex + Vertices.Num(); Vertex != DataEnd; ++Vertex)
+	{
+		Vertex->Normal.Normalize();
+	}
 }
 
 void LidarPointCloudMeshing::CalculateNormals(FLidarPointCloudOctree* Octree, FThreadSafeBool* bCancelled, int32 Quality, float Tolerance, TArray64<FLidarPointCloudPoint*>& InPointSelection)
@@ -219,7 +933,7 @@ void LidarPointCloudMeshing::CalculateNormals(FLidarPointCloudOctree* Octree, FT
 					}
 				}
 
-				bSuccess = (PointCount - SamplingUnit->Points.Num()) > 0;
+				bSuccess = PointCount - SamplingUnit->Points.Num() > 0;
 			}
 
 			if(!bSuccess)
@@ -264,3 +978,82 @@ void LidarPointCloudMeshing::CalculateNormals(FLidarPointCloudOctree* Octree, FT
 		SamplingUnit = nullptr;
 	}
 }
+
+void LidarPointCloudMeshing::BuildCollisionMesh(FLidarPointCloudOctree* Octree, const float& CellSize, FTriMeshCollisionData* CollisionMesh)
+{
+	if (!CollisionMesh)
+	{
+		return;
+	}
+	
+	TArray<FVector3f> Vertices;
+	MarchingCubes::Run(Octree, CellSize, false, Vertices, nullptr, nullptr);
+
+	FScopeBenchmarkTimer Timer("Exporting Collision Data");
+	
+	CollisionMesh->Vertices.Empty(Vertices.Num());
+	CollisionMesh->Indices.Empty(Vertices.Num() / 3);
+	CollisionMesh->Indices.AddUninitialized(Vertices.Num() / 3);
+
+	bool bInSet = false;
+	TSet<FVector3f> TmpVertexData;
+	TmpVertexData.Reserve(Vertices.Num());
+	uint32* IndicesPtr = (uint32*)CollisionMesh->Indices.GetData();
+	for (FVector3f *Vertex = Vertices.GetData(), *DataEnd = Vertex + Vertices.Num(); Vertex != DataEnd; ++Vertex, ++IndicesPtr)
+	{
+		*IndicesPtr = TmpVertexData.Add(*Vertex, &bInSet).AsInteger();
+		if (!bInSet)
+		{
+			CollisionMesh->Vertices.Add(*Vertex);
+		}
+	}
+
+	CollisionMesh->Vertices.Shrink();
+}
+
+void LidarPointCloudMeshing::BuildStaticMeshBuffers(FLidarPointCloudOctree* Octree, const float& CellSize, bool bUseSelection, FMeshBuffers* OutMeshBuffers, const FTransform& Transform)
+{
+	if(OutMeshBuffers == nullptr)
+	{
+		return;
+	}
+	
+	TArray<FVector3f> Vertices;
+	TArray<FVector3f> Normals;
+	TArray<FColor> Colors;
+	MarchingCubes::Run(Octree, CellSize, bUseSelection, Vertices, &Normals, &Colors);
+	
+	FScopeBenchmarkTimer Timer("Building Static Mesh Buffers");
+
+	FBox Bounds(EForceInit::ForceInit);
+	uint32 IndexOffset;
+	uint32 VertexOffset;
+	OutMeshBuffers->Init(Vertices.Num(), true, IndexOffset, VertexOffset);
+
+	// We do not support doubles yet
+	const FTransform3f Transform3f = (FTransform3f)Transform;
+
+	bool bInSet = false;
+	TSet<FIntVector3> TmpVertexData;
+	TmpVertexData.Reserve(Vertices.Num());
+	uint32* IndicesPtr = OutMeshBuffers->Indices.GetData() + IndexOffset;
+	FVector3f* NormalsPtr = Normals.GetData();
+	FColor* ColorsPtr = Colors.GetData();
+	for (FVector3f *Vertex = Vertices.GetData(), *DataEnd = Vertex + Vertices.Num(); Vertex != DataEnd; ++Vertex, ++IndicesPtr, ++NormalsPtr, ++ColorsPtr)
+	{
+		*IndicesPtr = TmpVertexData.Add(FIntVector3((FVector)(*Vertex * 100)), &bInSet).AsInteger() + VertexOffset;
+		if (bInSet)
+		{
+			OutMeshBuffers->Vertices[*IndicesPtr].Normal += *NormalsPtr;
+		}
+		else
+		{
+			const FVector3f Location = Transform3f.TransformPosition(*Vertex);
+			OutMeshBuffers->Vertices.Emplace(Location, *NormalsPtr, *ColorsPtr);
+			Bounds += (FVector)Location;
+		}
+	};
+	
+	OutMeshBuffers->ExpandBounds(Bounds);
+	OutMeshBuffers->NormalizeNormals();
+}
diff --git a/Engine/Plugins/Enterprise/LidarPointCloud/Source/LidarPointCloudRuntime/Public/LidarPointCloud.h b/Engine/Plugins/Enterprise/LidarPointCloud/Source/LidarPointCloudRuntime/Public/LidarPointCloud.h
--- a/Engine/Plugins/Enterprise/LidarPointCloud/Source/LidarPointCloudRuntime/Public/LidarPointCloud.h
+++ b/Engine/Plugins/Enterprise/LidarPointCloud/Source/LidarPointCloudRuntime/Public/LidarPointCloud.h
@@ -22,6 +22,11 @@ class UBodySetup;
 class FLidarPointCloudCollisionRendering;
 class FLidarPointCloudNotification;
 
+namespace LidarPointCloudMeshing
+{
+	struct FMeshBuffers;
+};
+
 /**
  * Used for ULidarPointCloud::CreateFromXXXX calls
  */
@@ -417,9 +422,11 @@ public:
 	void DeleteSelected();
 	void InvertSelection();
 	int64 NumSelectedPoints();
+	bool HasSelectedPoints();
 	void GetSelectedPointsAsCopies(TArray64<FLidarPointCloudPoint>& SelectedPoints, FTransform Transform);
 	void CalculateNormalsForSelection();
 	void ClearSelection();
+	void BuildStaticMeshBuffersForSelection(float CellSize, LidarPointCloudMeshing::FMeshBuffers* OutMeshBuffers, const FTransform& Transform);
 #endif
 	
 	UFUNCTION(BlueprintCallable, Category = "Lidar Point Cloud")
@@ -450,6 +457,9 @@ public:
 	UFUNCTION(BlueprintCallable, Category = "Lidar Point Cloud")
 	void RemoveCollision();
 
+	/** Constructs and returns the MeshBuffers struct from the data */
+	void BuildStaticMeshBuffers(float CellSize, LidarPointCloudMeshing::FMeshBuffers* OutMeshBuffers, const FTransform& Transform);
+	
 	/** Returns true, if the cloud is fully and persistently loaded. */
 	UFUNCTION(BlueprintPure, Category = "Lidar Point Cloud")
 	bool IsFullyLoaded() const { return Octree.IsFullyLoaded(); }
diff --git a/Engine/Plugins/Enterprise/LidarPointCloud/Source/LidarPointCloudRuntime/Public/LidarPointCloudOctree.h b/Engine/Plugins/Enterprise/LidarPointCloud/Source/LidarPointCloudRuntime/Public/LidarPointCloudOctree.h
--- a/Engine/Plugins/Enterprise/LidarPointCloud/Source/LidarPointCloudRuntime/Public/LidarPointCloudOctree.h
+++ b/Engine/Plugins/Enterprise/LidarPointCloud/Source/LidarPointCloudRuntime/Public/LidarPointCloudOctree.h
@@ -17,10 +17,14 @@ class ULidarPointCloud;
 class FLidarPointCloudOctree;
 struct FLidarPointCloudTraversalOctree;
 struct FLidarPointCloudTraversalOctreeNode;
-
 class FLidarPointCloudRenderBuffer;
 class FLidarPointCloudVertexFactory;
 
+namespace LidarPointCloudMeshing
+{
+	struct FMeshBuffers;
+};
+
 /**
  * WARNING: Exercise caution when modifying the contents of the Octree, as it may be in use by the Rendering Thread via FPointCloudSceneProxy
  * Use the FLidarPointCloudOctree::DataLock prior to such attempt
@@ -354,6 +358,9 @@ public:
 	/** Returns pointer to the collision data */
 	const FTriMeshCollisionData* GetCollisionData() const { return &CollisionMesh; }
 
+	/** Constructs and returns the MeshBuffers struct from the data */
+	void BuildStaticMeshBuffers(float CellSize, LidarPointCloudMeshing::FMeshBuffers* OutMeshBuffers, const FTransform& Transform);
+
 	/** Populates the given array with points from the tree */
 	void GetPoints(TArray<FLidarPointCloudPoint*>& SelectedPoints, int64 StartIndex = 0, int64 Count = -1);
 	void GetPoints(TArray64<FLidarPointCloudPoint*>& SelectedPoints, int64 StartIndex = 0, int64 Count = -1);
@@ -477,9 +484,12 @@ public:
 	void DeleteSelected();
 	void InvertSelection();
 	int64 NumSelectedPoints();
+	bool HasSelectedPoints();
 	void GetSelectedPointsAsCopies(TArray64<FLidarPointCloudPoint>& SelectedPoints, const FTransform& Transform);
+	void GetSelectedPointsInBox(TArray64<const FLidarPointCloudPoint*>& SelectedPoints, const FBox& Box) const;
 	void CalculateNormalsForSelection(FThreadSafeBool* bCancelled, int32 Quality, float Tolerance);
 	void ClearSelection();
+	void BuildStaticMeshBuffersForSelection(float CellSize, LidarPointCloudMeshing::FMeshBuffers* OutMeshBuffers, const FTransform& Transform);
 #endif
 
 	/** Initializes the Octree properties. */
diff --git a/Engine/Plugins/Enterprise/LidarPointCloud/Source/LidarPointCloudRuntime/Public/Meshing/LidarPointCloudMeshing.h b/Engine/Plugins/Enterprise/LidarPointCloud/Source/LidarPointCloudRuntime/Public/Meshing/LidarPointCloudMeshing.h
--- a/Engine/Plugins/Enterprise/LidarPointCloud/Source/LidarPointCloudRuntime/Public/Meshing/LidarPointCloudMeshing.h
+++ b/Engine/Plugins/Enterprise/LidarPointCloud/Source/LidarPointCloudRuntime/Public/Meshing/LidarPointCloudMeshing.h
@@ -6,8 +6,33 @@
 #include "LidarPointCloudShared.h"
 
 class FLidarPointCloudOctree;
+struct FTriMeshCollisionData;
+struct FMeshDescription;
 
 namespace LidarPointCloudMeshing
 {
+	struct FVertexData
+	{
+		FVector3f Position;
+		FVector3f Normal;
+		FColor Color;
+
+		FVertexData() {}
+		FVertexData(const FVector3f& Position, const FVector3f& Normal, const FColor& Color) : Position(Position), Normal(Normal), Color(Color) {}
+	};
+	
+	struct FMeshBuffers
+	{
+		TArray<uint32> Indices;
+		TArray<FVertexData> Vertices;
+		FBoxSphereBounds Bounds;
+		
+		void Init(uint32 Capacity, bool bExpandIfNotEmpty, uint32& OutIndexOffset, uint32& OutVertexOffset);
+		void ExpandBounds(const FBox& NewBounds);
+		void NormalizeNormals();
+	};
+	
 	void CalculateNormals(FLidarPointCloudOctree* Octree, FThreadSafeBool* bCancelled, int32 Quality, float Tolerance, TArray64<FLidarPointCloudPoint*>& InPointSelection);
+	void BuildCollisionMesh(FLidarPointCloudOctree* Octree, const float& CellSize, FTriMeshCollisionData* CollisionMesh);
+	void BuildStaticMeshBuffers(FLidarPointCloudOctree* Octree, const float& CellSize, bool bUseSelection, FMeshBuffers* OutMeshBuffers, const FTransform& Transform);
 };
\ No newline at end of file
