Corrected class names

diff --git a/Engine/Plugins/Enterprise/LidarPointCloud/Source/LidarPointCloudEditor/Private/LidarPointCloudEdMode.cpp b/Engine/Plugins/Enterprise/LidarPointCloud/Source/LidarPointCloudEditor/Private/LidarPointCloudEdMode.cpp
--- a/Engine/Plugins/Enterprise/LidarPointCloud/Source/LidarPointCloudEditor/Private/LidarPointCloudEdMode.cpp
+++ b/Engine/Plugins/Enterprise/LidarPointCloud/Source/LidarPointCloudEditor/Private/LidarPointCloudEdMode.cpp
@@ -38,7 +38,7 @@ void ULidarEditorMode::Enter()
 
 	const FLidarPointCloudEditorCommands& Commands = FLidarPointCloudEditorCommands::Get();
 
-#define REGISTER_TOOL(Tool) RegisterTool(Commands.Toolkit##Tool, TEXT("Lidar"#Tool"Tool"), NewObject<ULidarEditorToolBuilder_##Tool>())
+#define REGISTER_TOOL(Tool) RegisterTool(Commands.Toolkit##Tool, TEXT("Lidar"#Tool"Tool"), NewObject<ULidarEditorToolBuilder##Tool>())
 	REGISTER_TOOL(Select);
 	REGISTER_TOOL(Align);
 	REGISTER_TOOL(Merge);
@@ -123,7 +123,7 @@ void ULidarEditorMode::BindCommands()
 
 void ULidarEditorMode::OnToolStarted(UInteractiveToolManager* Manager, UInteractiveTool* Tool)
 {
-	if(ULidarEditorTool_Base* LidarTool = Cast<ULidarEditorTool_Base>(Tool))
+	if(ULidarEditorToolBase* LidarTool = Cast<ULidarEditorToolBase>(Tool))
 	{
 		GetToolManager()->DisplayMessage(LidarTool->GetToolMessage(), EToolMessageLevel::UserNotification);
 	}
@@ -131,7 +131,7 @@ void ULidarEditorMode::OnToolStarted(UInteractiveToolManager* Manager, UInteract
 
 void ULidarEditorMode::OnToolEnded(UInteractiveToolManager* Manager, UInteractiveTool* Tool)
 {
-	if(Tool && Tool->IsA(ULidarEditorTool_SelectionBase::StaticClass()))
+	if(Tool && Tool->IsA(ULidarEditorToolSelectionBase::StaticClass()))
 	{
 		FLidarPointCloudEditorHelper::ClearSelection();
 	}
diff --git a/Engine/Plugins/Enterprise/LidarPointCloud/Source/LidarPointCloudEditor/Private/LidarPointCloudEditorTools.cpp b/Engine/Plugins/Enterprise/LidarPointCloud/Source/LidarPointCloudEditor/Private/LidarPointCloudEditorTools.cpp
--- a/Engine/Plugins/Enterprise/LidarPointCloud/Source/LidarPointCloudEditor/Private/LidarPointCloudEditorTools.cpp
+++ b/Engine/Plugins/Enterprise/LidarPointCloud/Source/LidarPointCloudEditor/Private/LidarPointCloudEditorTools.cpp
@@ -22,7 +22,7 @@ namespace {
 	constexpr float PaintMaxDeviation = 0.15f;
 }
 
-void ULidarEditorTool_Base::Setup()
+void ULidarEditorToolBase::Setup()
 {
 	Super::Setup();
 
@@ -34,12 +34,12 @@ void ULidarEditorTool_Base::Setup()
 	}
 }
 
-FText ULidarEditorTool_Base::GetToolMessage() const
+FText ULidarEditorToolBase::GetToolMessage() const
 {
 	return FText();
 }
 
-void ULidarEditorTool_ClickDragBase::Setup()
+void ULidarEditorToolClickDragBase::Setup()
 {
 	Super::Setup();
 
@@ -56,14 +56,14 @@ void ULidarEditorTool_ClickDragBase::Setup()
 	AddInputBehavior(ClickDragBehavior, this);
 }
 
-void ULidarEditorTool_ClickDragBase::Render(IToolsContextRenderAPI* RenderAPI)
+void ULidarEditorToolClickDragBase::Render(IToolsContextRenderAPI* RenderAPI)
 {
 	Super::Render(RenderAPI);
 	
 	GetToolManager()->GetContextQueriesAPI()->GetCurrentViewState(CameraState);
 }
 
-void ULidarEditorTool_ClickDragBase::OnUpdateModifierState(int ModifierID, bool bIsOn)
+void ULidarEditorToolClickDragBase::OnUpdateModifierState(int ModifierID, bool bIsOn)
 {
 	switch (ModifierID)
 	{
@@ -78,105 +78,105 @@ void ULidarEditorTool_ClickDragBase::OnUpdateModifierState(int ModifierID, bool
 	}
 }
 
-void ULidarToolActions_Align::AlignAroundWorldOrigin()
+void ULidarToolActionsAlign::AlignAroundWorldOrigin()
 {
 	FLidarPointCloudEditorHelper::AlignSelectionAroundWorldOrigin();
 }
 
-void ULidarToolActions_Align::AlignAroundOriginalCoordinates()
+void ULidarToolActionsAlign::AlignAroundOriginalCoordinates()
 {
 	FLidarPointCloudEditorHelper::SetOriginalCoordinateForSelection();
 }
 
-void ULidarToolActions_Align::ResetAlignment()
+void ULidarToolActionsAlign::ResetAlignment()
 {
 	FLidarPointCloudEditorHelper::CenterSelection();
 }
 
-void ULidarToolActions_Merge::MergeActors()
+void ULidarToolActionsMerge::MergeActors()
 {
 	FLidarPointCloudEditorHelper::MergeSelectionByComponent(bReplaceSourceActorsAfterMerging);
 }
 
-void ULidarToolActions_Merge::MergeData()
+void ULidarToolActionsMerge::MergeData()
 {
 	FLidarPointCloudEditorHelper::MergeSelectionByData(bReplaceSourceActorsAfterMerging);
 }
 
-void ULidarToolActions_Collision::BuildCollision()
+void ULidarToolActionsCollision::BuildCollision()
  {
  	FLidarPointCloudEditorHelper::SetCollisionErrorForSelection(OverrideMaxCollisionError);
  	FLidarPointCloudEditorHelper::BuildCollisionForSelection();
  }
 
-void ULidarToolActions_Collision::RemoveCollision()
+void ULidarToolActionsCollision::RemoveCollision()
  {
  	FLidarPointCloudEditorHelper::RemoveCollisionForSelection();
  }
 
-void ULidarToolActions_Meshing::BuildStaticMesh()
+void ULidarToolActionsMeshing::BuildStaticMesh()
 {
 	FLidarPointCloudEditorHelper::MeshSelected(false, MaxMeshingError, bMergeMeshes, !bMergeMeshes && bRetainTransform);
 }
 
-void ULidarToolActions_Normals::CalculateNormals()
+void ULidarToolActionsNormals::CalculateNormals()
 {
 	FLidarPointCloudEditorHelper::SetNormalsQualityForSelection(Quality, NoiseTolerance);
 	FLidarPointCloudEditorHelper::CalculateNormalsForSelection();
 }
 
-void ULidarToolActions_Selection::HideSelected()
+void ULidarToolActionsSelection::HideSelected()
 {
 	FLidarPointCloudEditorHelper::HideSelected();
 }
 
-void ULidarToolActions_Selection::ResetVisibility()
+void ULidarToolActionsSelection::ResetVisibility()
 {
 	FLidarPointCloudEditorHelper::ResetVisibility();
 }
 
-void ULidarToolActions_Selection::DeleteHidden()
+void ULidarToolActionsSelection::DeleteHidden()
 {
 	FLidarPointCloudEditorHelper::DeleteHidden();
 }
 
-void ULidarToolActions_Selection::Extract()
+void ULidarToolActionsSelection::Extract()
 {
 	FLidarPointCloudEditorHelper::Extract();
 }
 
-void ULidarToolActions_Selection::ExtractAsCopy()
+void ULidarToolActionsSelection::ExtractAsCopy()
 {
 	FLidarPointCloudEditorHelper::ExtractAsCopy();
 }
 
-void ULidarToolActions_Selection::CalculateNormals()
+void ULidarToolActionsSelection::CalculateNormals()
 {
 	FLidarPointCloudEditorHelper::CalculateNormals();
 	ClearSelection();
 }
 
-void ULidarToolActions_Selection::DeleteSelected()
+void ULidarToolActionsSelection::DeleteSelected()
 {
 	FLidarPointCloudEditorHelper::DeleteSelected();
 }
 
-void ULidarToolActions_Selection::InvertSelection()
+void ULidarToolActionsSelection::InvertSelection()
 {
 	FLidarPointCloudEditorHelper::InvertSelection();
 }
 
-void ULidarToolActions_Selection::ClearSelection()
+void ULidarToolActionsSelection::ClearSelection()
 {
 	FLidarPointCloudEditorHelper::ClearSelection();
 }
 
-void ULidarToolActions_Selection::BuildStaticMesh()
+void ULidarToolActionsSelection::BuildStaticMesh()
 {
 	FLidarPointCloudEditorHelper::MeshSelected(true, MaxMeshingError, bMergeMeshes, !bMergeMeshes && bRetainTransform);
 }
 
-void ULidarEditorTool_SelectionBase::DrawHUD(FCanvas* Canvas, IToolsContextRenderAPI* RenderAPI)
+void ULidarEditorToolSelectionBase::DrawHUD(FCanvas* Canvas, IToolsContextRenderAPI* RenderAPI)
 {
 	if(Clicks.Num() == 0)
 	{
@@ -197,63 +197,63 @@ void ULidarEditorTool_SelectionBase::DrawHUD(FCanvas* Canvas, IToolsContextRende
 	Canvas->DrawItem(LineItem);
 }
 
-TArray<FConvexVolume> ULidarEditorTool_SelectionBase::GetSelectionConvexVolumes()
+TArray<FConvexVolume> ULidarEditorToolSelectionBase::GetSelectionConvexVolumes()
 {
 	return FLidarPointCloudEditorHelper::BuildConvexVolumesFromPoints(Clicks);
 }
 
-TObjectPtr<UInteractiveToolPropertySet> ULidarEditorTool_SelectionBase::CreateToolActions()
+TObjectPtr<UInteractiveToolPropertySet> ULidarEditorToolSelectionBase::CreateToolActions()
 {
-	return NewObject<ULidarToolActions_Selection>(this);
+	return NewObject<ULidarToolActionsSelection>(this);
 }
 
-bool ULidarEditorTool_SelectionBase::OnUpdateHover(const FInputDeviceRay& DevicePos)
+bool ULidarEditorToolSelectionBase::OnUpdateHover(const FInputDeviceRay& DevicePos)
 {
 	CurrentMousePos = DevicePos.ScreenPosition;
 	PostCurrentMousePosChanged();
 	return true;
 }
 
-void ULidarEditorTool_SelectionBase::OnClickPress(const FInputDeviceRay& PressPos)
+void ULidarEditorToolSelectionBase::OnClickPress(const FInputDeviceRay& PressPos)
 {
 	bSelecting = true;
 }
 
-void ULidarEditorTool_SelectionBase::OnClickRelease(const FInputDeviceRay& ReleasePos)
+void ULidarEditorToolSelectionBase::OnClickRelease(const FInputDeviceRay& ReleasePos)
 {
 	bSelecting = false;
 }
 
-void ULidarEditorTool_SelectionBase::OnClickDrag(const FInputDeviceRay& DragPos)
+void ULidarEditorToolSelectionBase::OnClickDrag(const FInputDeviceRay& DragPos)
 {
 	CurrentMousePos = DragPos.ScreenPosition;
 	PostCurrentMousePosChanged();
 }
 
-void ULidarEditorTool_SelectionBase::OnTerminateDragSequence()
+void ULidarEditorToolSelectionBase::OnTerminateDragSequence()
 {
 	Clicks.Empty();
 	bSelecting = false;
 }
 
-bool ULidarEditorTool_SelectionBase::ExecuteNestedCancelCommand()
+bool ULidarEditorToolSelectionBase::ExecuteNestedCancelCommand()
 {
 	OnTerminateDragSequence();
 	return true;
 }
 
-FText ULidarEditorTool_SelectionBase::GetToolMessage() const
+FText ULidarEditorToolSelectionBase::GetToolMessage() const
 {
 	const FText ToolMessage = LOCTEXT("ULidarEditorToolToolMessage", "Use Left-click to start the selection. Hold Shift to add selection, hold Ctrl to subtract selection.");
 	return ToolMessage;
 }
 
-FLinearColor ULidarEditorTool_SelectionBase::GetHUDColor()
+FLinearColor ULidarEditorToolSelectionBase::GetHUDColor()
 {
 	return FLinearColor::White;
 }
 
-void ULidarEditorTool_SelectionBase::FinalizeSelection()
+void ULidarEditorToolSelectionBase::FinalizeSelection()
 {
 	const ELidarPointCloudSelectionMode SelectionMode = GetSelectionMode();
 
@@ -270,7 +270,7 @@ void ULidarEditorTool_SelectionBase::FinalizeSelection()
 	}
 }
 
-ELidarPointCloudSelectionMode ULidarEditorTool_SelectionBase::GetSelectionMode() const
+ELidarPointCloudSelectionMode ULidarEditorToolSelectionBase::GetSelectionMode() const
 {
 	ELidarPointCloudSelectionMode SelectionMode = ELidarPointCloudSelectionMode::None;
 	if(bCtrlToggle)
@@ -284,12 +284,12 @@ ELidarPointCloudSelectionMode ULidarEditorTool_SelectionBase::GetSelectionMode()
 	return SelectionMode;
 }
 
-TArray<FConvexVolume> ULidarEditorTool_BoxSelection::GetSelectionConvexVolumes()
+TArray<FConvexVolume> ULidarEditorToolBoxSelection::GetSelectionConvexVolumes()
 {
 	return { FLidarPointCloudEditorHelper::BuildConvexVolumeFromCoordinates(Clicks[0], Clicks[2]) };
 }
 
-void ULidarEditorTool_BoxSelection::OnClickPress(const FInputDeviceRay& PressPos)
+void ULidarEditorToolBoxSelection::OnClickPress(const FInputDeviceRay& PressPos)
 {
 	Clicks.Append({
 		PressPos.ScreenPosition,
@@ -299,7 +299,7 @@ void ULidarEditorTool_BoxSelection::OnClickPress(const FInputDeviceRay& PressPos
 	});
 }
 
-void ULidarEditorTool_BoxSelection::OnClickDrag(const FInputDeviceRay& DragPos)
+void ULidarEditorToolBoxSelection::OnClickDrag(const FInputDeviceRay& DragPos)
 {
 	if(Clicks.Num() == 0)
 	{
@@ -311,7 +311,7 @@ void ULidarEditorTool_BoxSelection::OnClickDrag(const FInputDeviceRay& DragPos)
 	Clicks[3].X = DragPos.ScreenPosition.X;
 }
 
-void ULidarEditorTool_BoxSelection::OnClickRelease(const FInputDeviceRay& ReleasePos)
+void ULidarEditorToolBoxSelection::OnClickRelease(const FInputDeviceRay& ReleasePos)
 {
 	if(Clicks.Num() == 4)
 	{
@@ -328,7 +328,7 @@ void ULidarEditorTool_BoxSelection::OnClickRelease(const FInputDeviceRay& Releas
 	Clicks.Empty();
 }
 
-void ULidarEditorTool_PolygonalSelection::OnClickPress(const FInputDeviceRay& PressPos)
+void ULidarEditorToolPolygonalSelection::OnClickPress(const FInputDeviceRay& PressPos)
 {
 	if(IsWithinSnap())
 	{
@@ -341,12 +341,12 @@ void ULidarEditorTool_PolygonalSelection::OnClickPress(const FInputDeviceRay& Pr
 	}
 }
 
-FLinearColor ULidarEditorTool_PolygonalSelection::GetHUDColor()
+FLinearColor ULidarEditorToolPolygonalSelection::GetHUDColor()
 {
 	return IsWithinSnap() ? FLinearColor::Green : Super::GetHUDColor();
 }
 
-void ULidarEditorTool_PolygonalSelection::PostCurrentMousePosChanged()
+void ULidarEditorToolPolygonalSelection::PostCurrentMousePosChanged()
 {
 	if(IsWithinSnap())
 	{
@@ -354,17 +354,17 @@ void ULidarEditorTool_PolygonalSelection::PostCurrentMousePosChanged()
 	}
 }
 
-bool ULidarEditorTool_PolygonalSelection::IsWithinSnap()
+bool ULidarEditorToolPolygonalSelection::IsWithinSnap()
 {
 	return Clicks.Num() > 1 && (CurrentMousePos - Clicks[0]).SquaredLength() <= PolySnapDistanceSq;
 }
 
-void ULidarEditorTool_LassoSelection::OnClickPress(const FInputDeviceRay& PressPos)
+void ULidarEditorToolLassoSelection::OnClickPress(const FInputDeviceRay& PressPos)
 {
 	Clicks.Add(PressPos.ScreenPosition);
 }
 
-void ULidarEditorTool_LassoSelection::OnClickDrag(const FInputDeviceRay& DragPos)
+void ULidarEditorToolLassoSelection::OnClickDrag(const FInputDeviceRay& DragPos)
 {
 	Super::OnClickDrag(DragPos);
 	
@@ -374,7 +374,7 @@ void ULidarEditorTool_LassoSelection::OnClickDrag(const FInputDeviceRay& DragPos
 	}
 }
 
-void ULidarEditorTool_LassoSelection::OnClickRelease(const FInputDeviceRay& ReleasePos)
+void ULidarEditorToolLassoSelection::OnClickRelease(const FInputDeviceRay& ReleasePos)
 {
 	if(Clicks.Num() > 1)
 	{
@@ -384,13 +384,13 @@ void ULidarEditorTool_LassoSelection::OnClickRelease(const FInputDeviceRay& Rele
 	Clicks.Empty();
 }
 
-void ULidarEditorTool_PaintSelection::Setup()
+void ULidarEditorToolPaintSelection::Setup()
 {
 	Super::Setup();
-	BrushRadius = GetDefault<ULidarToolActions_PaintSelection>()->BrushRadius;
+	BrushRadius = GetDefault<ULidarToolActionsPaintSelection>()->BrushRadius;
 }
 
-void ULidarEditorTool_PaintSelection::Render(IToolsContextRenderAPI* RenderAPI)
+void ULidarEditorToolPaintSelection::Render(IToolsContextRenderAPI* RenderAPI)
 {
 	Super::Render(RenderAPI);
 	
@@ -400,7 +400,7 @@ void ULidarEditorTool_PaintSelection::Render(IToolsContextRenderAPI* RenderAPI)
 	}
 }
 
-void ULidarEditorTool_PaintSelection::PostCurrentMousePosChanged()
+void ULidarEditorToolPaintSelection::PostCurrentMousePosChanged()
 {
 	const FLidarPointCloudRay Ray = FLidarPointCloudEditorHelper::MakeRayFromScreenPosition(CurrentMousePos);
 
@@ -429,7 +429,7 @@ void ULidarEditorTool_PaintSelection::PostCurrentMousePosChanged()
 	LastHitDistance = NewDistance;
 }
 
-void ULidarEditorTool_PaintSelection::OnClickPress(const FInputDeviceRay& PressPos)
+void ULidarEditorToolPaintSelection::OnClickPress(const FInputDeviceRay& PressPos)
 {
 	Super::OnClickPress(PressPos);
 
@@ -441,15 +441,15 @@ void ULidarEditorTool_PaintSelection::OnClickPress(const FInputDeviceRay& PressP
 	Paint();
 }
 
-void ULidarEditorTool_PaintSelection::OnClickDrag(const FInputDeviceRay& DragPos)
+void ULidarEditorToolPaintSelection::OnClickDrag(const FInputDeviceRay& DragPos)
 {
 	Super::OnClickDrag(DragPos);
 	Paint();
 }
 
-void ULidarEditorTool_PaintSelection::OnPropertyModified(UObject* PropertySet, FProperty* Property)
+void ULidarEditorToolPaintSelection::OnPropertyModified(UObject* PropertySet, FProperty* Property)
 {
-	if(const ULidarToolActions_PaintSelection* Actions = Cast<ULidarToolActions_PaintSelection>(PropertySet))
+	if(const ULidarToolActionsPaintSelection* Actions = Cast<ULidarToolActionsPaintSelection>(PropertySet))
 	{
 		if(Property && Property->GetName().Equals("BrushRadius"))
 		{
@@ -458,7 +458,7 @@ void ULidarEditorTool_PaintSelection::OnPropertyModified(UObject* PropertySet, F
 	}
 }
 
-void ULidarEditorTool_PaintSelection::Paint()
+void ULidarEditorToolPaintSelection::Paint()
 {
 	if(bHasHit)
 	{
diff --git a/Engine/Plugins/Enterprise/LidarPointCloud/Source/LidarPointCloudEditor/Private/LidarPointCloudEditorTools.h b/Engine/Plugins/Enterprise/LidarPointCloud/Source/LidarPointCloudEditor/Private/LidarPointCloudEditorTools.h
--- a/Engine/Plugins/Enterprise/LidarPointCloud/Source/LidarPointCloudEditor/Private/LidarPointCloudEditorTools.h
+++ b/Engine/Plugins/Enterprise/LidarPointCloud/Source/LidarPointCloudEditor/Private/LidarPointCloudEditorTools.h
@@ -14,7 +14,7 @@ class URectangleMarqueeMechanic;
 class UClickDragInputBehavior;
 
 UCLASS()
-class ULidarEditorTool_Base : public UInteractiveTool
+class ULidarEditorToolBase : public UInteractiveTool
 {
 public:
 	GENERATED_BODY()
@@ -27,18 +27,18 @@ public:
 };
 
 UCLASS()
-class ULidarEditorToolBuilder_Base : public UInteractiveToolWithToolTargetsBuilder
+class ULidarEditorToolBuilderBase : public UInteractiveToolWithToolTargetsBuilder
 {
 	GENERATED_BODY()
 
 public:
 	virtual bool CanBuildTool(const FToolBuilderState& SceneState) const override { return true; }
-	virtual UInteractiveTool* BuildTool(const FToolBuilderState& SceneState) const override { return NewObject<ULidarEditorTool_Base>(SceneState.ToolManager); }
+	virtual UInteractiveTool* BuildTool(const FToolBuilderState& SceneState) const override { return NewObject<ULidarEditorToolBase>(SceneState.ToolManager); }
 };
 
 UCLASS()
-class ULidarEditorTool_ClickDragBase :
-	public ULidarEditorTool_Base,
+class ULidarEditorToolClickDragBase :
+	public ULidarEditorToolBase,
 	public IClickDragBehaviorTarget,
 	public IHoverBehaviorTarget,
 	public IInteractiveToolNestedAcceptCancelAPI
@@ -84,13 +84,13 @@ protected:
 // -------------------------------------------------------
 
 UCLASS()
-class ULidarEditorToolBuilder_Select : public ULidarEditorToolBuilder_Base
+class ULidarEditorToolBuilderSelect : public ULidarEditorToolBuilderBase
 {
 	GENERATED_BODY()
 };
 
 UCLASS()
-class ULidarToolActions_Align : public UInteractiveToolPropertySet
+class ULidarToolActionsAlign : public UInteractiveToolPropertySet
 {
 	GENERATED_BODY()
 
@@ -106,23 +106,23 @@ public:
 };
 
 UCLASS()
-class ULidarEditorTool_Align : public ULidarEditorTool_Base
+class ULidarEditorToolAlign : public ULidarEditorToolBase
 {
 public:
 	GENERATED_BODY()
-	virtual TObjectPtr<UInteractiveToolPropertySet> CreateToolActions() override { return NewObject<ULidarToolActions_Align>(this); }
+	virtual TObjectPtr<UInteractiveToolPropertySet> CreateToolActions() override { return NewObject<ULidarToolActionsAlign>(this); }
 };
 
 UCLASS()
-class ULidarEditorToolBuilder_Align : public ULidarEditorToolBuilder_Base
+class ULidarEditorToolBuilderAlign : public ULidarEditorToolBuilderBase
 {
 public:
 	GENERATED_BODY()
-	virtual UInteractiveTool* BuildTool(const FToolBuilderState& SceneState) const override { return NewObject<ULidarEditorTool_Align>(SceneState.ToolManager); }
+	virtual UInteractiveTool* BuildTool(const FToolBuilderState& SceneState) const override { return NewObject<ULidarEditorToolAlign>(SceneState.ToolManager); }
 };
 
 UCLASS()
-class ULidarToolActions_Merge : public UInteractiveToolPropertySet
+class ULidarToolActionsMerge : public UInteractiveToolPropertySet
 {
 	GENERATED_BODY()
 
@@ -138,23 +138,23 @@ public:
 };
 
 UCLASS()
-class ULidarEditorTool_Merge : public ULidarEditorTool_Base
+class ULidarEditorToolMerge : public ULidarEditorToolBase
 {
 public:
 	GENERATED_BODY()
-	virtual TObjectPtr<UInteractiveToolPropertySet> CreateToolActions() override { return NewObject<ULidarToolActions_Merge>(this); }
+	virtual TObjectPtr<UInteractiveToolPropertySet> CreateToolActions() override { return NewObject<ULidarToolActionsMerge>(this); }
 };
 
 UCLASS()
-class ULidarEditorToolBuilder_Merge : public ULidarEditorToolBuilder_Base
+class ULidarEditorToolBuilderMerge : public ULidarEditorToolBuilderBase
 {
 public:
 	GENERATED_BODY()
-	virtual UInteractiveTool* BuildTool(const FToolBuilderState& SceneState) const override { return NewObject<ULidarEditorTool_Merge>(SceneState.ToolManager); }
+	virtual UInteractiveTool* BuildTool(const FToolBuilderState& SceneState) const override { return NewObject<ULidarEditorToolMerge>(SceneState.ToolManager); }
 };
 
 UCLASS()
-class ULidarToolActions_Collision : public UInteractiveToolPropertySet
+class ULidarToolActionsCollision : public UInteractiveToolPropertySet
 {
 	GENERATED_BODY()
 
@@ -170,23 +170,23 @@ public:
 };
 
 UCLASS()
-class ULidarEditorTool_Collision : public ULidarEditorTool_Base
+class ULidarEditorToolCollision : public ULidarEditorToolBase
 {
 public:
 	GENERATED_BODY()
-	virtual TObjectPtr<UInteractiveToolPropertySet> CreateToolActions() override { return NewObject<ULidarToolActions_Collision>(this); }
+	virtual TObjectPtr<UInteractiveToolPropertySet> CreateToolActions() override { return NewObject<ULidarToolActionsCollision>(this); }
 };
 
 UCLASS()
-class ULidarEditorToolBuilder_Collision : public ULidarEditorToolBuilder_Base
+class ULidarEditorToolBuilderCollision : public ULidarEditorToolBuilderBase
 {
 public:
 	GENERATED_BODY()
-	virtual UInteractiveTool* BuildTool(const FToolBuilderState& SceneState) const override { return NewObject<ULidarEditorTool_Collision>(SceneState.ToolManager); }
+	virtual UInteractiveTool* BuildTool(const FToolBuilderState& SceneState) const override { return NewObject<ULidarEditorToolCollision>(SceneState.ToolManager); }
 };
 
 UCLASS()
-class ULidarToolActions_Meshing : public UInteractiveToolPropertySet
+class ULidarToolActionsMeshing : public UInteractiveToolPropertySet
 {
 	GENERATED_BODY()
 
@@ -207,23 +207,23 @@ public:
 };
 
 UCLASS()
-class ULidarEditorTool_Meshing : public ULidarEditorTool_Base
+class ULidarEditorToolMeshing : public ULidarEditorToolBase
 {
 public:
 	GENERATED_BODY()
-	virtual TObjectPtr<UInteractiveToolPropertySet> CreateToolActions() override { return NewObject<ULidarToolActions_Meshing>(this); }
+	virtual TObjectPtr<UInteractiveToolPropertySet> CreateToolActions() override { return NewObject<ULidarToolActionsMeshing>(this); }
 };
 
 UCLASS()
-class ULidarEditorToolBuilder_Meshing : public ULidarEditorToolBuilder_Base
+class ULidarEditorToolBuilderMeshing : public ULidarEditorToolBuilderBase
 {
 public:
 	GENERATED_BODY()
-	virtual UInteractiveTool* BuildTool(const FToolBuilderState& SceneState) const override { return NewObject<ULidarEditorTool_Meshing>(SceneState.ToolManager); }
+	virtual UInteractiveTool* BuildTool(const FToolBuilderState& SceneState) const override { return NewObject<ULidarEditorToolMeshing>(SceneState.ToolManager); }
 };
 
 UCLASS()
-class ULidarToolActions_Normals : public UInteractiveToolPropertySet
+class ULidarToolActionsNormals : public UInteractiveToolPropertySet
 {
 	GENERATED_BODY()
 
@@ -245,23 +245,23 @@ public:
 };
 
 UCLASS()
-class ULidarEditorTool_Normals : public ULidarEditorTool_Base
+class ULidarEditorToolNormals : public ULidarEditorToolBase
 {
 public:
 	GENERATED_BODY()
-	virtual TObjectPtr<UInteractiveToolPropertySet> CreateToolActions() override { return NewObject<ULidarToolActions_Normals>(this); }
+	virtual TObjectPtr<UInteractiveToolPropertySet> CreateToolActions() override { return NewObject<ULidarToolActionsNormals>(this); }
 };
 
 UCLASS()
-class ULidarEditorToolBuilder_Normals : public ULidarEditorToolBuilder_Base
+class ULidarEditorToolBuilderNormals : public ULidarEditorToolBuilderBase
 {
 public:
 	GENERATED_BODY()
-	virtual UInteractiveTool* BuildTool(const FToolBuilderState& SceneState) const override { return NewObject<ULidarEditorTool_Normals>(SceneState.ToolManager); }
+	virtual UInteractiveTool* BuildTool(const FToolBuilderState& SceneState) const override { return NewObject<ULidarEditorToolNormals>(SceneState.ToolManager); }
 };
 
 UCLASS()
-class ULidarToolActions_Selection : public UInteractiveToolPropertySet
+class ULidarToolActionsSelection : public UInteractiveToolPropertySet
 {
 	GENERATED_BODY()
 
@@ -309,7 +309,7 @@ public:
 };
 
 UCLASS()
-class ULidarEditorTool_SelectionBase : public ULidarEditorTool_ClickDragBase
+class ULidarEditorToolSelectionBase : public ULidarEditorToolClickDragBase
 {
 public:
 	GENERATED_BODY()
@@ -336,7 +336,7 @@ protected:
 };
 
 UCLASS()
-class ULidarEditorTool_BoxSelection : public ULidarEditorTool_SelectionBase
+class ULidarEditorToolBoxSelection : public ULidarEditorToolSelectionBase
 {
 public:
 	GENERATED_BODY()
@@ -347,15 +347,15 @@ public:
 };
 
 UCLASS()
-class ULidarEditorToolBuilder_BoxSelection : public ULidarEditorToolBuilder_Base
+class ULidarEditorToolBuilderBoxSelection : public ULidarEditorToolBuilderBase
 {
 public:
 	GENERATED_BODY()
-	virtual UInteractiveTool* BuildTool(const FToolBuilderState& SceneState) const override { return NewObject<ULidarEditorTool_BoxSelection>(SceneState.ToolManager); }
+	virtual UInteractiveTool* BuildTool(const FToolBuilderState& SceneState) const override { return NewObject<ULidarEditorToolBoxSelection>(SceneState.ToolManager); }
 };
 
 UCLASS()
-class ULidarEditorTool_PolygonalSelection : public ULidarEditorTool_SelectionBase
+class ULidarEditorToolPolygonalSelection : public ULidarEditorToolSelectionBase
 {
 public:
 	GENERATED_BODY()
@@ -368,15 +368,15 @@ private:
 };
 
 UCLASS()
-class ULidarEditorToolBuilder_PolygonalSelection : public ULidarEditorToolBuilder_Base
+class ULidarEditorToolBuilderPolygonalSelection : public ULidarEditorToolBuilderBase
 {
 public:
 	GENERATED_BODY()
-	virtual UInteractiveTool* BuildTool(const FToolBuilderState& SceneState) const override { return NewObject<ULidarEditorTool_PolygonalSelection>(SceneState.ToolManager); }
+	virtual UInteractiveTool* BuildTool(const FToolBuilderState& SceneState) const override { return NewObject<ULidarEditorToolPolygonalSelection>(SceneState.ToolManager); }
 };
 
 UCLASS()
-class ULidarEditorTool_LassoSelection : public ULidarEditorTool_SelectionBase
+class ULidarEditorToolLassoSelection : public ULidarEditorToolSelectionBase
 {
 public:
 	GENERATED_BODY()
@@ -386,15 +386,15 @@ public:
 };
 
 UCLASS()
-class ULidarEditorToolBuilder_LassoSelection : public ULidarEditorToolBuilder_Base
+class ULidarEditorToolBuilderLassoSelection : public ULidarEditorToolBuilderBase
 {
 public:
 	GENERATED_BODY()
-	virtual UInteractiveTool* BuildTool(const FToolBuilderState& SceneState) const override { return NewObject<ULidarEditorTool_LassoSelection>(SceneState.ToolManager); }
+	virtual UInteractiveTool* BuildTool(const FToolBuilderState& SceneState) const override { return NewObject<ULidarEditorToolLassoSelection>(SceneState.ToolManager); }
 };
 
 UCLASS()
-class ULidarToolActions_PaintSelection : public ULidarToolActions_Selection
+class ULidarToolActionsPaintSelection : public ULidarToolActionsSelection
 {
 	GENERATED_BODY()
 
@@ -404,7 +404,7 @@ public:
 };
 
 UCLASS()
-class ULidarEditorTool_PaintSelection : public ULidarEditorTool_SelectionBase
+class ULidarEditorToolPaintSelection : public ULidarEditorToolSelectionBase
 {
 public:
 	GENERATED_BODY()
@@ -414,7 +414,7 @@ public:
 	virtual void PostCurrentMousePosChanged() override;
 	virtual void OnClickPress(const FInputDeviceRay& PressPos) override;
 	virtual void OnClickDrag(const FInputDeviceRay& DragPos) override;
-	virtual TObjectPtr<UInteractiveToolPropertySet> CreateToolActions() override { return NewObject<ULidarToolActions_PaintSelection>(this); }
+	virtual TObjectPtr<UInteractiveToolPropertySet> CreateToolActions() override { return NewObject<ULidarToolActionsPaintSelection>(this); }
 	virtual void OnPropertyModified(UObject* PropertySet, FProperty* Property) override;
 
 private:
@@ -428,9 +428,9 @@ private:
 };
 
 UCLASS()
-class ULidarEditorToolBuilder_PaintSelection : public ULidarEditorToolBuilder_Base
+class ULidarEditorToolBuilderPaintSelection : public ULidarEditorToolBuilderBase
 {
 public:
 	GENERATED_BODY()
-	virtual UInteractiveTool* BuildTool(const FToolBuilderState& SceneState) const override { return NewObject<ULidarEditorTool_PaintSelection>(SceneState.ToolManager); }
+	virtual UInteractiveTool* BuildTool(const FToolBuilderState& SceneState) const override { return NewObject<ULidarEditorToolPaintSelection>(SceneState.ToolManager); }
 };
\ No newline at end of file
