Lidar: Added latent OnComplete for BuildCollision function

diff --git a/Engine/Plugins/Enterprise/LidarPointCloud/Source/LidarPointCloudRuntime/Private/LidarPointCloud.cpp b/Engine/Plugins/Enterprise/LidarPointCloud/Source/LidarPointCloudRuntime/Private/LidarPointCloud.cpp
--- a/Engine/Plugins/Enterprise/LidarPointCloud/Source/LidarPointCloudRuntime/Private/LidarPointCloud.cpp
+++ b/Engine/Plugins/Enterprise/LidarPointCloud/Source/LidarPointCloudRuntime/Private/LidarPointCloud.cpp
@@ -89,6 +89,29 @@ private:
 /////////////////////////////////////////////////
 // FPointCloudLatentAction
 
+class FPointCloudSimpleLatentAction : public FPendingLatentAction
+{
+	FName ExecutionFunction;
+	int32 OutputLink;
+	FWeakObjectPtr CallbackTarget;
+
+public:
+	bool bComplete;
+
+	FPointCloudSimpleLatentAction(const FLatentActionInfo& LatentInfo)
+		: ExecutionFunction(LatentInfo.ExecutionFunction)
+		, OutputLink(LatentInfo.Linkage)
+		, CallbackTarget(LatentInfo.CallbackTarget)
+		, bComplete(false)
+	{
+	}
+
+	virtual void UpdateOperation(FLatentResponse& Response) override
+	{
+		Response.FinishAndTriggerIf(bComplete, ExecutionFunction, OutputLink, CallbackTarget);
+	}
+};
+
 class FPointCloudLatentAction : public FPendingLatentAction
 {
 public:
@@ -663,11 +686,35 @@ void ULidarPointCloud::SetOptimalCollisionError()
 	MaxCollisionError = FMath::CeilToInt(Octree.GetEstimatedPointSpacing() * 300) * 0.01f;
 }
 
-void ULidarPointCloud::BuildCollision()
+void ULidarPointCloud::BuildCollisionWithCallback(UObject* WorldContextObject, FLatentActionInfo LatentInfo, bool& bSuccess)
+{
+	if (UWorld* World = GEngine->GetWorldFromContextObject(WorldContextObject, EGetWorldErrorMode::LogAndReturnNull))
+	{
+		FLatentActionManager& LatentActionManager = World->GetLatentActionManager();
+		if (LatentActionManager.FindExistingAction<FPointCloudLatentAction>(LatentInfo.CallbackTarget, LatentInfo.UUID) == nullptr)
+		{
+			FPointCloudSimpleLatentAction* CompletionAction = new FPointCloudSimpleLatentAction(LatentInfo);
+
+			LatentActionManager.AddNewAction(LatentInfo.CallbackTarget, LatentInfo.UUID, CompletionAction);
+
+			BuildCollision([CompletionAction, &bSuccess](bool bCompletedSuccessfully)
+			{
+				bSuccess = bCompletedSuccessfully;
+				CompletionAction->bComplete = true;
+			});
+		}
+	}
+}
+
+void ULidarPointCloud::BuildCollision(TFunction<void(bool)> CompletionCallback)
 {
 	if (bCollisionBuildInProgress)
 	{
 		PC_ERROR("Another collision operation already in progress.");
+		if(CompletionCallback)
+		{
+			CompletionCallback(false);
+		}
 		return;
 	}
 
@@ -683,7 +730,7 @@ void ULidarPointCloud::BuildCollision()
 
 	TWeakObjectPtr<ULidarPointCloud> WeakThis = this;
 
-	Async(EAsyncExecution::Thread, [WeakThis = MoveTemp(WeakThis), Notification]() mutable
+	Async(EAsyncExecution::Thread, [WeakThis = MoveTemp(WeakThis), Notification, CompletionCallback = MoveTemp(CompletionCallback)]() mutable
 	{
 		if (ULidarPointCloud* RawPtr = WeakThis.Get())
 		{
@@ -693,11 +740,15 @@ void ULidarPointCloud::BuildCollision()
 			FBenchmarkTimer::Reset();
 
 			RawPtr->NewBodySetup->CreatePhysicsMeshes();
-			AsyncTask(ENamedThreads::GameThread, [WeakThis = MoveTemp(WeakThis), Notification]
+			AsyncTask(ENamedThreads::GameThread, [WeakThis = MoveTemp(WeakThis), Notification, CompletionCallback = MoveTemp(CompletionCallback)]
 			{
 				if (ULidarPointCloud* RawPtr = WeakThis.Get())
 				{
 					RawPtr->FinishPhysicsAsyncCook(true, Notification);
+					if(CompletionCallback)
+					{
+						CompletionCallback(true);
+					}
 				}
 			});
 		}
diff --git a/Engine/Plugins/Enterprise/LidarPointCloud/Source/LidarPointCloudRuntime/Public/LidarPointCloud.h b/Engine/Plugins/Enterprise/LidarPointCloud/Source/LidarPointCloudRuntime/Public/LidarPointCloud.h
--- a/Engine/Plugins/Enterprise/LidarPointCloud/Source/LidarPointCloudRuntime/Public/LidarPointCloud.h
+++ b/Engine/Plugins/Enterprise/LidarPointCloud/Source/LidarPointCloudRuntime/Public/LidarPointCloud.h
@@ -451,7 +451,11 @@ public:
 
 	/** Builds collision mesh for the cloud, using current collision settings */
 	UFUNCTION(BlueprintCallable, Category = "Lidar Point Cloud")
-	void BuildCollision();
+	void BuildCollision() { BuildCollision(nullptr); }
+	void BuildCollision(TFunction<void(bool)> CompletionCallback);
+
+	UFUNCTION(BlueprintCallable, Category = "Lidar Point Cloud", meta = (Latent, WorldContext = "WorldContextObject", LatentInfo = "LatentInfo"))
+	void BuildCollisionWithCallback(UObject* WorldContextObject, FLatentActionInfo LatentInfo, bool& bSuccess);
 
 	/** Removes collision mesh from the cloud. */
 	UFUNCTION(BlueprintCallable, Category = "Lidar Point Cloud")
