LiDAR: Move plugin from Beta to Production ready

diff --git a/Engine/Plugins/Enterprise/LidarPointCloud/LidarPointCloud.uplugin b/Engine/Plugins/Enterprise/LidarPointCloud/LidarPointCloud.uplugin
--- a/Engine/Plugins/Enterprise/LidarPointCloud/LidarPointCloud.uplugin
+++ b/Engine/Plugins/Enterprise/LidarPointCloud/LidarPointCloud.uplugin
@@ -1,7 +1,7 @@
 {
 	"FileVersion": 1,
 	"Version": 1,
-	"VersionName": "0.8",
+	"VersionName": "1.0",
 	"FriendlyName": "LiDAR Point Cloud Support",
 	"Description": "Adds support for importing, processing and rendering of LiDAR Point Clouds.",
 	"Category": "Rendering",
@@ -12,7 +12,7 @@
 	"SupportURL": "",
 	"EnabledByDefault": false,
 	"CanContainContent": true,
-	"IsBetaVersion": true,
+	"IsBetaVersion": false,
 	"Installed": false,
 	"SupportedTargetPlatforms": [
 		"Win64",
