Fixed CIS warnings

diff --git a/Engine/Plugins/Enterprise/LidarPointCloud/Source/LidarPointCloudEditor/Private/LidarPointCloudEditorHelper.cpp b/Engine/Plugins/Enterprise/LidarPointCloud/Source/LidarPointCloudEditor/Private/LidarPointCloudEditorHelper.cpp
--- a/Engine/Plugins/Enterprise/LidarPointCloud/Source/LidarPointCloudEditor/Private/LidarPointCloudEditorHelper.cpp
+++ b/Engine/Plugins/Enterprise/LidarPointCloud/Source/LidarPointCloudEditor/Private/LidarPointCloudEditorHelper.cpp
@@ -494,7 +494,7 @@ void FLidarPointCloudEditorHelper::RemoveCollisionForSelection()
 
 void FLidarPointCloudEditorHelper::MeshSelected(bool bMeshByPoints, float CellSize, bool bMergeMeshes, bool bRetainTransform)
 {
-	const int32 NumSteps = (bMeshByPoints ? GetNumLidarActors() : GEditor->GetSelectedActorCount()) + bMergeMeshes;
+	const int32 NumSteps = (bMeshByPoints ? GetNumLidarActors() : GEditor->GetSelectedActorCount()) + (bMergeMeshes ? 1 : 0);
 	
 	FScopedSlowTask ProgressDialog(NumSteps, LOCTEXT("Meshing", "Meshing Point Clouds..."));
 	ProgressDialog.MakeDialog();
