Revert to before CL 21090768.

diff --git a/Engine/Plugins/Enterprise/LidarPointCloud/Source/LidarPointCloudEditor/Private/LidarPointCloudEditorViewportClient.cpp b/Engine/Plugins/Enterprise/LidarPointCloud/Source/LidarPointCloudEditor/Private/LidarPointCloudEditorViewportClient.cpp
--- a/Engine/Plugins/Enterprise/LidarPointCloud/Source/LidarPointCloudEditor/Private/LidarPointCloudEditorViewportClient.cpp
+++ b/Engine/Plugins/Enterprise/LidarPointCloud/Source/LidarPointCloudEditor/Private/LidarPointCloudEditorViewportClient.cpp
@@ -52,6 +52,47 @@ TArray<FVector2D> ToVectorArray(const TArray<FIntPoint>& Points)
 	return VectorPoints;
 }
 
+// Slow, O(n2), but sufficient for the current problem
+bool IsPolygonSelfIntersecting(const TArray<FVector2D>& Points, bool bAllowLooping)
+{
+	const int32 MaxIndex = bAllowLooping ? Points.Num() : Points.Num() - 1;
+
+	for (int32 i = 0; i < MaxIndex; ++i)
+	{
+		const int32 i1 = i < Points.Num() - 1 ? i + 1 : 0;
+
+		const FVector2D P1 = Points[i];
+		const FVector2D P2 = Points[i1];
+
+		for (int32 j = 0; j < MaxIndex; ++j)
+		{
+			const int32 j1 = j < Points.Num() - 1 ? j + 1 : 0;
+
+			if (j1 != i && j != i && j != i1)
+			{
+				// Modified, inlined FMath::SegmentIntersection2D
+				const FVector2D SegmentStartA = P1;
+				const FVector2D SegmentEndA = P2;
+				const FVector2D SegmentStartB = Points[j];
+				const FVector2D SegmentEndB = Points[j1];
+				const FVector2D VectorA = P2 - SegmentStartA;
+				const FVector2D VectorB = SegmentEndB - SegmentStartB;
+
+				const float S = (-VectorA.Y * (SegmentStartA.X - SegmentStartB.X) + VectorA.X * (SegmentStartA.Y - SegmentStartB.Y)) / (-VectorB.X * VectorA.Y + VectorA.X * VectorB.Y);
+				const float T = (VectorB.X * (SegmentStartA.Y - SegmentStartB.Y) - VectorB.Y * (SegmentStartA.X - SegmentStartB.X)) / (-VectorB.X * VectorA.Y + VectorA.X * VectorB.Y);
+
+				if (S >= 0 && S <= 1 && T >= 0 && T <= 1)
+				{
+					return true;
+				}
+			}
+		}
+	}
+
+	return false;
+}
+
+
 // Copied from GeomTools.cpp and converted to work with FIntPoint
 bool IsPolygonConvex(const TArray<FIntPoint>& Points)
 {
@@ -175,23 +216,23 @@ void FLidarPointCloudEditorViewportClient::Tick(float DeltaSeconds)
 	}
 }
 
-bool FLidarPointCloudEditorViewportClient::InputKey(const FInputKeyEventArgs& EventArgs)
+bool FLidarPointCloudEditorViewportClient::InputKey(FViewport* InViewport, int32 ControllerId, FKey Key, EInputEvent Event, float AmountDepressed, bool Gamepad)
 {
 	using UE::Lidar::Private::Editor::PolySnapDistanceSq;
 
 	bool bHandled = false;
 
-	const bool bAlt = EventArgs.Key == EKeys::LeftAlt || EventArgs.Key == EKeys::RightAlt;
-	const bool bCtrl = EventArgs.Key == EKeys::LeftControl || EventArgs.Key == EKeys::RightControl;
+	const bool bAlt = Key == EKeys::LeftAlt || Key == EKeys::RightAlt;
+	const bool bCtrl = Key == EKeys::LeftControl || Key == EKeys::RightControl;
 
 	if (TSharedPtr<FLidarPointCloudEditor> Editor = PointCloudEditorPtr.Pin())
 	{
 		// Edit Mode
 		if (Editor->IsEditMode())
 		{
-			if (EventArgs.Event == IE_Pressed)
+			if (Event == IE_Pressed)
 			{
-				if (EventArgs.Key == EKeys::Delete)
+				if (Key == EKeys::Delete)
 				{
 					if (Viewport->KeyState(EKeys::LeftShift) || Viewport->KeyState(EKeys::RightShift))
 					{
@@ -204,13 +245,13 @@ bool FLidarPointCloudEditorViewportClient::InputKey(const FInputKeyEventArgs& Ev
 
 					bHandled = true;
 				}
-				else if (EventArgs.Key == EKeys::Escape)
+				else if (Key == EKeys::Escape)
 				{
 					Editor->DeselectPoints();
 					SelectionPoints.Empty();
 					bHandled = true;
 				}
-				else if (EventArgs.Key == EKeys::Enter)
+				else if (Key == EKeys::Enter)
 				{
 					if (SelectionMethod == ELidarPointCloudSelectionMethod::Polygonal)
 					{
@@ -229,7 +270,7 @@ bool FLidarPointCloudEditorViewportClient::InputKey(const FInputKeyEventArgs& Ev
 					bHandled = true;
 					SelectionMode = ELidarPointCloudSelectionMode::Add;
 				}
-				else if (EventArgs.Key == EKeys::LeftMouseButton)
+				else if (Key == EKeys::LeftMouseButton)
 				{
 					bHandled = true;
 
@@ -237,13 +278,13 @@ bool FLidarPointCloudEditorViewportClient::InputKey(const FInputKeyEventArgs& Ev
 					if (SelectionMethod == ELidarPointCloudSelectionMethod::Box)
 					{
 						// Mark the cursor location for selection start
-						EventArgs.Viewport->GetMousePos(SelectionPoints[SelectionPoints.AddUninitialized()]);
+						InViewport->GetMousePos(SelectionPoints[SelectionPoints.AddUninitialized()]);
 					}
 					else if (SelectionMethod == ELidarPointCloudSelectionMethod::Polygonal)
 					{
 						// Add new lasso point
 						FIntPoint NewPoint;
-						EventArgs.Viewport->GetMousePos(NewPoint);
+						InViewport->GetMousePos(NewPoint);
 
 						// Don't allow duplicates
 						if (SelectionPoints.Num() == 0 || SelectionPoints.Last() != NewPoint)
@@ -251,7 +292,7 @@ bool FLidarPointCloudEditorViewportClient::InputKey(const FInputKeyEventArgs& Ev
 							TArray<FVector2D> VectorPoints = ToVectorArray(SelectionPoints);
 							VectorPoints.Add(NewPoint);
 
-							if (!FLidarPointCloudEditorHelper::IsPolygonSelfIntersecting(VectorPoints, false))
+							if (!IsPolygonSelfIntersecting(VectorPoints, false))
 							{
 								// Snap to first point
 								if (SelectionPoints.Num() > 1 && (NewPoint - SelectionPoints[0]).SizeSquared() < PolySnapDistanceSq)
@@ -268,7 +309,7 @@ bool FLidarPointCloudEditorViewportClient::InputKey(const FInputKeyEventArgs& Ev
 					}
 					else if (SelectionMethod == ELidarPointCloudSelectionMethod::Lasso)
 					{
-						EventArgs.Viewport->GetMousePos(SelectionPoints[SelectionPoints.AddUninitialized()]);
+						InViewport->GetMousePos(SelectionPoints[SelectionPoints.AddUninitialized()]);
 					}
 					else if (SelectionMethod == ELidarPointCloudSelectionMethod::Paint)
 					{
@@ -281,12 +322,12 @@ bool FLidarPointCloudEditorViewportClient::InputKey(const FInputKeyEventArgs& Ev
 					// Do not block ability to change camera speed
 					if (!Viewport->KeyState(EKeys::RightMouseButton))
 					{
-						if (EventArgs.Key == EKeys::MouseScrollUp)
+						if (Key == EKeys::MouseScrollUp)
 						{
 							PaintingRadius *= 1.1f;
 							bHandled = true;
 						}
-						else if (EventArgs.Key == EKeys::MouseScrollDown)
+						else if (Key == EKeys::MouseScrollDown)
 						{
 							PaintingRadius /= 1.1f;
 							bHandled = true;
@@ -294,21 +335,21 @@ bool FLidarPointCloudEditorViewportClient::InputKey(const FInputKeyEventArgs& Ev
 					}
 				}
 			}
-			else if (EventArgs.Event == IE_Released)
+			else if (Event == IE_Released)
 			{
 				if (bAlt || bCtrl)
 				{
 					SelectionMode = ELidarPointCloudSelectionMode::None;
 					bHandled = true;
 				}
-				else if (EventArgs.Key == EKeys::LeftMouseButton)
+				else if (Key == EKeys::LeftMouseButton)
 				{
 					bHandled = true;
 
 					if (SelectionMethod == ELidarPointCloudSelectionMethod::Box)
 					{
 						// Mark the cursor location for selection end
-						EventArgs.Viewport->GetMousePos(SelectionPoints[SelectionPoints.AddUninitialized()]);
+						InViewport->GetMousePos(SelectionPoints[SelectionPoints.AddUninitialized()]);
 						OnBoxSelectionEnd();
 						SelectionPoints.Empty();
 					}
@@ -329,19 +370,19 @@ bool FLidarPointCloudEditorViewportClient::InputKey(const FInputKeyEventArgs& Ev
 
 	if (!bHandled)
 	{
-		bHandled = FEditorViewportClient::InputKey(EventArgs);
+		bHandled = FEditorViewportClient::InputKey(InViewport, ControllerId, Key, Event, AmountDepressed, false);
 
 		// Handle viewport screenshot.
-		bHandled |= InputTakeScreenshot(EventArgs.Viewport, EventArgs.Key, EventArgs.Event);
+		bHandled |= InputTakeScreenshot(InViewport, Key, Event);
 
-		bHandled |= AdvancedPreviewScene->HandleInputKey(EventArgs);
+		bHandled |= AdvancedPreviewScene->HandleInputKey(InViewport, ControllerId, Key, Event, AmountDepressed, Gamepad);
 
 	}
 
 	return bHandled;
 }
 
-bool FLidarPointCloudEditorViewportClient::InputAxis(FViewport* InViewport, FInputDeviceId DeviceId, FKey Key, float Delta, float DeltaTime, int32 NumSamples, bool bGamepad)
+bool FLidarPointCloudEditorViewportClient::InputAxis(FViewport* InViewport, int32 ControllerId, FKey Key, float Delta, float DeltaTime, int32 NumSamples, bool bGamepad)
 {
 	using UE::Lidar::Private::Editor::LassoSpacingSq;
 
@@ -379,14 +420,14 @@ bool FLidarPointCloudEditorViewportClient::InputAxis(FViewport* InViewport, FInp
 
 		if(!bHandled)
 		{
-			bHandled = AdvancedPreviewScene->HandleViewportInput(InViewport, DeviceId, Key, Delta, DeltaTime, NumSamples, bGamepad);
+			bHandled = AdvancedPreviewScene->HandleViewportInput(InViewport, ControllerId, Key, Delta, DeltaTime, NumSamples, bGamepad);
 			if (bHandled)
 			{
 				Invalidate();
 			}
 			else
 			{
-				bHandled = FEditorViewportClient::InputAxis(InViewport, DeviceId, Key, Delta, DeltaTime, NumSamples, bGamepad);
+				bHandled = FEditorViewportClient::InputAxis(InViewport, ControllerId, Key, Delta, DeltaTime, NumSamples, bGamepad);
 			}
 		}
 	}
@@ -649,7 +690,7 @@ void FLidarPointCloudEditorViewportClient::OnPolygonalSelectionEnd()
 			else
 			{
 				// Check for self-intersecting shape
-				if (!FLidarPointCloudEditorHelper::IsPolygonSelfIntersecting(VectorPoints, true))
+				if (!IsPolygonSelfIntersecting(VectorPoints, true))
 				{
 					// The separation needs points in CCW order
 					if (!FGeomTools2D::IsPolygonWindingCCW(VectorPoints))
@@ -780,7 +821,7 @@ void FLidarPointCloudEditorViewportClient::DrawSelectionPolygonal(FCanvas& Canva
 
 	// Calculate visual indication of complete polygon for the user
 	const bool bPolyComplete = DrawSelectionPoints.Num() > 2 && (DrawSelectionPoints.Last() - DrawSelectionPoints[0]).SizeSquared() < PolySnapDistanceSq;
-	const bool bSelfIntersecting = DrawSelectionPoints.Num() > 2 && FLidarPointCloudEditorHelper::IsPolygonSelfIntersecting(VectorPoints, true);
+	const bool bSelfIntersecting = DrawSelectionPoints.Num() > 2 && IsPolygonSelfIntersecting(VectorPoints, true);
 	FLinearColor SelectionColor = bSelfIntersecting ? FLinearColor::Red : bPolyComplete ? FLinearColor::Green : GetDefault<UEditorStyleSettings>()->SelectionColor;
 
 	// Selection Area
diff --git a/Engine/Plugins/Enterprise/LidarPointCloud/Source/LidarPointCloudEditor/Private/LidarPointCloudEditorViewportClient.h b/Engine/Plugins/Enterprise/LidarPointCloud/Source/LidarPointCloudEditor/Private/LidarPointCloudEditorViewportClient.h
--- a/Engine/Plugins/Enterprise/LidarPointCloud/Source/LidarPointCloudEditor/Private/LidarPointCloudEditorViewportClient.h
+++ b/Engine/Plugins/Enterprise/LidarPointCloud/Source/LidarPointCloudEditor/Private/LidarPointCloudEditorViewportClient.h
@@ -6,7 +6,7 @@
 #include "InputCoreTypes.h"
 #include "UnrealWidgetFwd.h"
 #include "EditorViewportClient.h"
-#include "LidarPointCloudEditorHelper.h"
+#include "Components.h"
 #include "LidarPointCloudShared.h"
 
 class FAdvancedPreviewScene;
@@ -16,6 +16,13 @@ class SLidarPointCloudEditorViewport;
 class ULidarPointCloud;
 class ULidarPointCloudComponent;
 
+enum class ELidarPointCloudSelectionMode : uint8
+{
+	None,
+	Add,
+	Subtract
+};
+
 enum class ELidarPointCloudSelectionMethod : uint8
 {
 	Box,
@@ -32,8 +39,8 @@ public:
 	~FLidarPointCloudEditorViewportClient();
 
 	// FEditorViewportClient interface
-	virtual bool InputKey(const FInputKeyEventArgs& EventArgs) override;
-	virtual bool InputAxis(FViewport* InViewport, FInputDeviceId DeviceId, FKey Key, float Delta, float DeltaTime, int32 NumSamples = 1, bool bGamepad = false) override;
+	virtual bool InputKey(FViewport* InViewport, int32 ControllerId, FKey Key, EInputEvent Event, float AmountDepressed = 1.f, bool bGamepad = false) override;
+	virtual bool InputAxis(FViewport* InViewport, int32 ControllerId, FKey Key, float Delta, float DeltaTime, int32 NumSamples = 1, bool bGamepad = false) override;
 	virtual void Tick(float DeltaSeconds) override;
 	virtual void DrawCanvas(FViewport& InViewport, FSceneView& View, FCanvas& Canvas) override;
 	virtual bool InputWidgetDelta(FViewport* InViewport, EAxisList::Type CurrentAxis, FVector& Drag, FRotator& Rot, FVector& Scale) override { return false; }
