Un-revert to CL 20745189

#rnx

diff --git a/Engine/Plugins/Enterprise/LidarPointCloud/Source/LidarPointCloudEditor/Private/LidarPointCloudEditorViewportClient.cpp b/Engine/Plugins/Enterprise/LidarPointCloud/Source/LidarPointCloudEditor/Private/LidarPointCloudEditorViewportClient.cpp
--- a/Engine/Plugins/Enterprise/LidarPointCloud/Source/LidarPointCloudEditor/Private/LidarPointCloudEditorViewportClient.cpp
+++ b/Engine/Plugins/Enterprise/LidarPointCloud/Source/LidarPointCloudEditor/Private/LidarPointCloudEditorViewportClient.cpp
@@ -52,47 +52,6 @@ TArray<FVector2D> ToVectorArray(const TArray<FIntPoint>& Points)
 	return VectorPoints;
 }
 
-// Slow, O(n2), but sufficient for the current problem
-bool IsPolygonSelfIntersecting(const TArray<FVector2D>& Points, bool bAllowLooping)
-{
-	const int32 MaxIndex = bAllowLooping ? Points.Num() : Points.Num() - 1;
-
-	for (int32 i = 0; i < MaxIndex; ++i)
-	{
-		const int32 i1 = i < Points.Num() - 1 ? i + 1 : 0;
-
-		const FVector2D P1 = Points[i];
-		const FVector2D P2 = Points[i1];
-
-		for (int32 j = 0; j < MaxIndex; ++j)
-		{
-			const int32 j1 = j < Points.Num() - 1 ? j + 1 : 0;
-
-			if (j1 != i && j != i && j != i1)
-			{
-				// Modified, inlined FMath::SegmentIntersection2D
-				const FVector2D SegmentStartA = P1;
-				const FVector2D SegmentEndA = P2;
-				const FVector2D SegmentStartB = Points[j];
-				const FVector2D SegmentEndB = Points[j1];
-				const FVector2D VectorA = P2 - SegmentStartA;
-				const FVector2D VectorB = SegmentEndB - SegmentStartB;
-
-				const float S = (-VectorA.Y * (SegmentStartA.X - SegmentStartB.X) + VectorA.X * (SegmentStartA.Y - SegmentStartB.Y)) / (-VectorB.X * VectorA.Y + VectorA.X * VectorB.Y);
-				const float T = (VectorB.X * (SegmentStartA.Y - SegmentStartB.Y) - VectorB.Y * (SegmentStartA.X - SegmentStartB.X)) / (-VectorB.X * VectorA.Y + VectorA.X * VectorB.Y);
-
-				if (S >= 0 && S <= 1 && T >= 0 && T <= 1)
-				{
-					return true;
-				}
-			}
-		}
-	}
-
-	return false;
-}
-
-
 // Copied from GeomTools.cpp and converted to work with FIntPoint
 bool IsPolygonConvex(const TArray<FIntPoint>& Points)
 {
@@ -216,23 +175,23 @@ void FLidarPointCloudEditorViewportClient::Tick(float DeltaSeconds)
 	}
 }
 
-bool FLidarPointCloudEditorViewportClient::InputKey(FViewport* InViewport, int32 ControllerId, FKey Key, EInputEvent Event, float AmountDepressed, bool Gamepad)
+bool FLidarPointCloudEditorViewportClient::InputKey(const FInputKeyEventArgs& EventArgs)
 {
 	using UE::Lidar::Private::Editor::PolySnapDistanceSq;
 
 	bool bHandled = false;
 
-	const bool bAlt = Key == EKeys::LeftAlt || Key == EKeys::RightAlt;
-	const bool bCtrl = Key == EKeys::LeftControl || Key == EKeys::RightControl;
+	const bool bAlt = EventArgs.Key == EKeys::LeftAlt || EventArgs.Key == EKeys::RightAlt;
+	const bool bCtrl = EventArgs.Key == EKeys::LeftControl || EventArgs.Key == EKeys::RightControl;
 
 	if (TSharedPtr<FLidarPointCloudEditor> Editor = PointCloudEditorPtr.Pin())
 	{
 		// Edit Mode
 		if (Editor->IsEditMode())
 		{
-			if (Event == IE_Pressed)
+			if (EventArgs.Event == IE_Pressed)
 			{
-				if (Key == EKeys::Delete)
+				if (EventArgs.Key == EKeys::Delete)
 				{
 					if (Viewport->KeyState(EKeys::LeftShift) || Viewport->KeyState(EKeys::RightShift))
 					{
@@ -245,13 +204,13 @@ bool FLidarPointCloudEditorViewportClient::InputKey(FViewport* InViewport, int32
 
 					bHandled = true;
 				}
-				else if (Key == EKeys::Escape)
+				else if (EventArgs.Key == EKeys::Escape)
 				{
 					Editor->DeselectPoints();
 					SelectionPoints.Empty();
 					bHandled = true;
 				}
-				else if (Key == EKeys::Enter)
+				else if (EventArgs.Key == EKeys::Enter)
 				{
 					if (SelectionMethod == ELidarPointCloudSelectionMethod::Polygonal)
 					{
@@ -270,7 +229,7 @@ bool FLidarPointCloudEditorViewportClient::InputKey(FViewport* InViewport, int32
 					bHandled = true;
 					SelectionMode = ELidarPointCloudSelectionMode::Add;
 				}
-				else if (Key == EKeys::LeftMouseButton)
+				else if (EventArgs.Key == EKeys::LeftMouseButton)
 				{
 					bHandled = true;
 
@@ -278,13 +237,13 @@ bool FLidarPointCloudEditorViewportClient::InputKey(FViewport* InViewport, int32
 					if (SelectionMethod == ELidarPointCloudSelectionMethod::Box)
 					{
 						// Mark the cursor location for selection start
-						InViewport->GetMousePos(SelectionPoints[SelectionPoints.AddUninitialized()]);
+						EventArgs.Viewport->GetMousePos(SelectionPoints[SelectionPoints.AddUninitialized()]);
 					}
 					else if (SelectionMethod == ELidarPointCloudSelectionMethod::Polygonal)
 					{
 						// Add new lasso point
 						FIntPoint NewPoint;
-						InViewport->GetMousePos(NewPoint);
+						EventArgs.Viewport->GetMousePos(NewPoint);
 
 						// Don't allow duplicates
 						if (SelectionPoints.Num() == 0 || SelectionPoints.Last() != NewPoint)
@@ -292,7 +251,7 @@ bool FLidarPointCloudEditorViewportClient::InputKey(FViewport* InViewport, int32
 							TArray<FVector2D> VectorPoints = ToVectorArray(SelectionPoints);
 							VectorPoints.Add(NewPoint);
 
-							if (!IsPolygonSelfIntersecting(VectorPoints, false))
+							if (!FLidarPointCloudEditorHelper::IsPolygonSelfIntersecting(VectorPoints, false))
 							{
 								// Snap to first point
 								if (SelectionPoints.Num() > 1 && (NewPoint - SelectionPoints[0]).SizeSquared() < PolySnapDistanceSq)
@@ -309,7 +268,7 @@ bool FLidarPointCloudEditorViewportClient::InputKey(FViewport* InViewport, int32
 					}
 					else if (SelectionMethod == ELidarPointCloudSelectionMethod::Lasso)
 					{
-						InViewport->GetMousePos(SelectionPoints[SelectionPoints.AddUninitialized()]);
+						EventArgs.Viewport->GetMousePos(SelectionPoints[SelectionPoints.AddUninitialized()]);
 					}
 					else if (SelectionMethod == ELidarPointCloudSelectionMethod::Paint)
 					{
@@ -322,12 +281,12 @@ bool FLidarPointCloudEditorViewportClient::InputKey(FViewport* InViewport, int32
 					// Do not block ability to change camera speed
 					if (!Viewport->KeyState(EKeys::RightMouseButton))
 					{
-						if (Key == EKeys::MouseScrollUp)
+						if (EventArgs.Key == EKeys::MouseScrollUp)
 						{
 							PaintingRadius *= 1.1f;
 							bHandled = true;
 						}
-						else if (Key == EKeys::MouseScrollDown)
+						else if (EventArgs.Key == EKeys::MouseScrollDown)
 						{
 							PaintingRadius /= 1.1f;
 							bHandled = true;
@@ -335,21 +294,21 @@ bool FLidarPointCloudEditorViewportClient::InputKey(FViewport* InViewport, int32
 					}
 				}
 			}
-			else if (Event == IE_Released)
+			else if (EventArgs.Event == IE_Released)
 			{
 				if (bAlt || bCtrl)
 				{
 					SelectionMode = ELidarPointCloudSelectionMode::None;
 					bHandled = true;
 				}
-				else if (Key == EKeys::LeftMouseButton)
+				else if (EventArgs.Key == EKeys::LeftMouseButton)
 				{
 					bHandled = true;
 
 					if (SelectionMethod == ELidarPointCloudSelectionMethod::Box)
 					{
 						// Mark the cursor location for selection end
-						InViewport->GetMousePos(SelectionPoints[SelectionPoints.AddUninitialized()]);
+						EventArgs.Viewport->GetMousePos(SelectionPoints[SelectionPoints.AddUninitialized()]);
 						OnBoxSelectionEnd();
 						SelectionPoints.Empty();
 					}
@@ -370,19 +329,19 @@ bool FLidarPointCloudEditorViewportClient::InputKey(FViewport* InViewport, int32
 
 	if (!bHandled)
 	{
-		bHandled = FEditorViewportClient::InputKey(InViewport, ControllerId, Key, Event, AmountDepressed, false);
+		bHandled = FEditorViewportClient::InputKey(EventArgs);
 
 		// Handle viewport screenshot.
-		bHandled |= InputTakeScreenshot(InViewport, Key, Event);
+		bHandled |= InputTakeScreenshot(EventArgs.Viewport, EventArgs.Key, EventArgs.Event);
 
-		bHandled |= AdvancedPreviewScene->HandleInputKey(InViewport, ControllerId, Key, Event, AmountDepressed, Gamepad);
+		bHandled |= AdvancedPreviewScene->HandleInputKey(EventArgs);
 
 	}
 
 	return bHandled;
 }
 
-bool FLidarPointCloudEditorViewportClient::InputAxis(FViewport* InViewport, int32 ControllerId, FKey Key, float Delta, float DeltaTime, int32 NumSamples, bool bGamepad)
+bool FLidarPointCloudEditorViewportClient::InputAxis(FViewport* InViewport, FInputDeviceId DeviceId, FKey Key, float Delta, float DeltaTime, int32 NumSamples, bool bGamepad)
 {
 	using UE::Lidar::Private::Editor::LassoSpacingSq;
 
@@ -420,14 +379,14 @@ bool FLidarPointCloudEditorViewportClient::InputAxis(FViewport* InViewport, int3
 
 		if(!bHandled)
 		{
-			bHandled = AdvancedPreviewScene->HandleViewportInput(InViewport, ControllerId, Key, Delta, DeltaTime, NumSamples, bGamepad);
+			bHandled = AdvancedPreviewScene->HandleViewportInput(InViewport, DeviceId, Key, Delta, DeltaTime, NumSamples, bGamepad);
 			if (bHandled)
 			{
 				Invalidate();
 			}
 			else
 			{
-				bHandled = FEditorViewportClient::InputAxis(InViewport, ControllerId, Key, Delta, DeltaTime, NumSamples, bGamepad);
+				bHandled = FEditorViewportClient::InputAxis(InViewport, DeviceId, Key, Delta, DeltaTime, NumSamples, bGamepad);
 			}
 		}
 	}
@@ -690,7 +649,7 @@ void FLidarPointCloudEditorViewportClient::OnPolygonalSelectionEnd()
 			else
 			{
 				// Check for self-intersecting shape
-				if (!IsPolygonSelfIntersecting(VectorPoints, true))
+				if (!FLidarPointCloudEditorHelper::IsPolygonSelfIntersecting(VectorPoints, true))
 				{
 					// The separation needs points in CCW order
 					if (!FGeomTools2D::IsPolygonWindingCCW(VectorPoints))
@@ -821,7 +780,7 @@ void FLidarPointCloudEditorViewportClient::DrawSelectionPolygonal(FCanvas& Canva
 
 	// Calculate visual indication of complete polygon for the user
 	const bool bPolyComplete = DrawSelectionPoints.Num() > 2 && (DrawSelectionPoints.Last() - DrawSelectionPoints[0]).SizeSquared() < PolySnapDistanceSq;
-	const bool bSelfIntersecting = DrawSelectionPoints.Num() > 2 && IsPolygonSelfIntersecting(VectorPoints, true);
+	const bool bSelfIntersecting = DrawSelectionPoints.Num() > 2 && FLidarPointCloudEditorHelper::IsPolygonSelfIntersecting(VectorPoints, true);
 	FLinearColor SelectionColor = bSelfIntersecting ? FLinearColor::Red : bPolyComplete ? FLinearColor::Green : GetDefault<UEditorStyleSettings>()->SelectionColor;
 
 	// Selection Area
diff --git a/Engine/Plugins/Enterprise/LidarPointCloud/Source/LidarPointCloudEditor/Private/LidarPointCloudEditorViewportClient.h b/Engine/Plugins/Enterprise/LidarPointCloud/Source/LidarPointCloudEditor/Private/LidarPointCloudEditorViewportClient.h
--- a/Engine/Plugins/Enterprise/LidarPointCloud/Source/LidarPointCloudEditor/Private/LidarPointCloudEditorViewportClient.h
+++ b/Engine/Plugins/Enterprise/LidarPointCloud/Source/LidarPointCloudEditor/Private/LidarPointCloudEditorViewportClient.h
@@ -6,7 +6,7 @@
 #include "InputCoreTypes.h"
 #include "UnrealWidgetFwd.h"
 #include "EditorViewportClient.h"
-#include "Components.h"
+#include "LidarPointCloudEditorHelper.h"
 #include "LidarPointCloudShared.h"
 
 class FAdvancedPreviewScene;
@@ -16,13 +16,6 @@ class SLidarPointCloudEditorViewport;
 class ULidarPointCloud;
 class ULidarPointCloudComponent;
 
-enum class ELidarPointCloudSelectionMode : uint8
-{
-	None,
-	Add,
-	Subtract
-};
-
 enum class ELidarPointCloudSelectionMethod : uint8
 {
 	Box,
@@ -39,8 +32,8 @@ public:
 	~FLidarPointCloudEditorViewportClient();
 
 	// FEditorViewportClient interface
-	virtual bool InputKey(FViewport* InViewport, int32 ControllerId, FKey Key, EInputEvent Event, float AmountDepressed = 1.f, bool bGamepad = false) override;
-	virtual bool InputAxis(FViewport* InViewport, int32 ControllerId, FKey Key, float Delta, float DeltaTime, int32 NumSamples = 1, bool bGamepad = false) override;
+	virtual bool InputKey(const FInputKeyEventArgs& EventArgs) override;
+	virtual bool InputAxis(FViewport* InViewport, FInputDeviceId DeviceId, FKey Key, float Delta, float DeltaTime, int32 NumSamples = 1, bool bGamepad = false) override;
 	virtual void Tick(float DeltaSeconds) override;
 	virtual void DrawCanvas(FViewport& InViewport, FSceneView& View, FCanvas& Canvas) override;
 	virtual bool InputWidgetDelta(FViewport* InViewport, EAxisList::Type CurrentAxis, FVector& Drag, FRotator& Rot, FVector& Scale) override { return false; }
