Inclusive Terminology

diff --git a/Engine/Plugins/Enterprise/LidarPointCloud/Source/LidarPointCloudRuntime/Private/LidarPointCloudComponent.cpp b/Engine/Plugins/Enterprise/LidarPointCloud/Source/LidarPointCloudRuntime/Private/LidarPointCloudComponent.cpp
--- a/Engine/Plugins/Enterprise/LidarPointCloud/Source/LidarPointCloudRuntime/Private/LidarPointCloudComponent.cpp
+++ b/Engine/Plugins/Enterprise/LidarPointCloud/Source/LidarPointCloudRuntime/Private/LidarPointCloudComponent.cpp
@@ -49,10 +49,10 @@ ULidarPointCloudComponent::ULidarPointCloudComponent()
 	SetCollisionProfileName(UCollisionProfile::BlockAll_ProfileName);
 
 	static ConstructorHelpers::FObjectFinder<UMaterial> M_PointCloud(TEXT("/LidarPointCloud/Materials/M_LidarPointCloud"));
-	MasterMaterial = M_PointCloud.Object;
+	BaseMaterial = M_PointCloud.Object;
 
 	static ConstructorHelpers::FObjectFinder<UMaterial> M_PointCloud_Masked(TEXT("/LidarPointCloud/Materials/M_LidarPointCloud_Masked"));
-	MasterMaterialMasked = M_PointCloud_Masked.Object;
+	BaseMaterialMasked = M_PointCloud_Masked.Object;
 }
 
 FBoxSphereBounds ULidarPointCloudComponent::CalcBounds(const FTransform& LocalToWorld) const
@@ -70,7 +70,7 @@ void ULidarPointCloudComponent::UpdateMaterial()
 	// ... otherwise, create MID from it
 	else
 	{
-		Material = UMaterialInstanceDynamic::Create(CustomMaterial ? CustomMaterial : PointShape != ELidarPointCloudSpriteShape::Square ? MasterMaterialMasked : MasterMaterial, nullptr);
+		Material = UMaterialInstanceDynamic::Create(CustomMaterial ? CustomMaterial : PointShape != ELidarPointCloudSpriteShape::Square ? BaseMaterialMasked : BaseMaterial, nullptr);
 	}
 
 	ApplyRenderingParameters();
diff --git a/Engine/Plugins/Enterprise/LidarPointCloud/Source/LidarPointCloudRuntime/Public/LidarPointCloudComponent.h b/Engine/Plugins/Enterprise/LidarPointCloud/Source/LidarPointCloudRuntime/Public/LidarPointCloudComponent.h
--- a/Engine/Plugins/Enterprise/LidarPointCloud/Source/LidarPointCloudRuntime/Public/LidarPointCloudComponent.h
+++ b/Engine/Plugins/Enterprise/LidarPointCloud/Source/LidarPointCloudRuntime/Public/LidarPointCloudComponent.h
@@ -27,7 +27,7 @@ class LIDARPOINTCLOUDRUNTIME_API ULidarPointCloudComponent : public UMeshCompone
 		
 private:
 	UPROPERTY(EditAnywhere, BlueprintReadOnly, Category = "Lidar Point Cloud", meta = (AllowPrivateAccess = "true"))
-	ULidarPointCloud *PointCloud;
+	ULidarPointCloud* PointCloud;
 
 	/**
 	 * Allows using custom-built material for the point cloud.
@@ -147,9 +147,9 @@ private:
 	UPROPERTY(Transient)
 	UMaterialInterface* Material;
 	UPROPERTY(Transient)
-	UMaterialInterface *MasterMaterial;
+	UMaterialInterface* BaseMaterial;
 	UPROPERTY(Transient)
-	UMaterialInterface* MasterMaterialMasked;
+	UMaterialInterface* BaseMaterialMasked;
 
 	/** Pointer to the viewport client of the owning editor, or null, if this is a game object. */
 	TWeakPtr<FViewportClient> OwningViewportClient;
