LiDARClippingVolume rendering incorrectly

diff --git a/Engine/Plugins/Enterprise/LidarPointCloud/Shaders/Private/LidarPointCloudVertexFactory.ush b/Engine/Plugins/Enterprise/LidarPointCloud/Shaders/Private/LidarPointCloudVertexFactory.ush
--- a/Engine/Plugins/Enterprise/LidarPointCloud/Shaders/Private/LidarPointCloudVertexFactory.ush
+++ b/Engine/Plugins/Enterprise/LidarPointCloud/Shaders/Private/LidarPointCloudVertexFactory.ush
@@ -268,8 +268,6 @@ float CalcSpriteSize(float3 Position)
 
 float3 ProcessPosition(float3 Position, half3x3 TangentToLocal, int VertexID)
 {
-	Position = ApplyClipping(Position + LocationOffset);
-	
 	const float Size = bUseScreenSizeScaling ? mul(float4(Position.xyz, 1), ResolvedView.TranslatedWorldToView).z * 0.01f : CalcSpriteSize(Position);
 	const half2 Scale = PrecomputedUV_Quad[VertexID % 4];
 	return Position + (TangentToLocal[0] * Scale.x + TangentToLocal[1] * Scale.y) * SpriteSizeMultiplier * Size;
@@ -327,7 +325,7 @@ FVertexFactoryIntermediates GetVertexFactoryIntermediates(FVertexFactoryInput In
 	Intermediates.TangentToLocal = CalcTangentToLocal(MetaData);
 	Intermediates.TangentToWorldSign = Intermediates.SceneData.InstanceData.DeterminantSign;
 	Intermediates.TangentToWorld = CalcTangentToWorldNoScale(Intermediates.TangentToLocal);
-	Intermediates.Position = Position;
+	Intermediates.Position = ApplyClipping(Position + LocationOffset);
 	
 	// Color
 	{
@@ -395,13 +393,13 @@ float4 VertexFactoryGetWorldPosition(FVertexFactoryInput Input, FVertexFactoryIn
 }
 float4 VertexFactoryGetWorldPosition(FPositionOnlyVertexFactoryInput Input)
 {
-	float4 WorldPosition = TransformLocalToTranslatedWorld(bUseStaticBuffers ? Input.Position.xyz : LidarGetInstancePosition(Input.VertexId));
+	float4 WorldPosition = TransformLocalToTranslatedWorld(ApplyClipping((bUseStaticBuffers ? Input.Position.xyz : LidarGetInstancePosition(Input.VertexId)) + LocationOffset));
 	WorldPosition.xyz = ProcessPosition(WorldPosition.xyz, Input.VertexId);
 	return WorldPosition;
 }
 float4 VertexFactoryGetWorldPosition(FPositionAndNormalOnlyVertexFactoryInput Input)
 {
-	float4 WorldPosition = TransformLocalToTranslatedWorld(bUseStaticBuffers ? Input.Position.xyz : LidarGetInstancePosition(Input.VertexId));
+	float4 WorldPosition = TransformLocalToTranslatedWorld(ApplyClipping((bUseStaticBuffers ? Input.Position.xyz : LidarGetInstancePosition(Input.VertexId)) + LocationOffset));
 	WorldPosition.xyz = ProcessPosition(WorldPosition.xyz, Input.VertexId);
 	return WorldPosition;
 }
diff --git a/Engine/Plugins/Enterprise/LidarPointCloud/Source/LidarPointCloudRuntime/Private/LidarPointCloudLODManager.cpp b/Engine/Plugins/Enterprise/LidarPointCloud/Source/LidarPointCloudRuntime/Private/LidarPointCloudLODManager.cpp
--- a/Engine/Plugins/Enterprise/LidarPointCloud/Source/LidarPointCloudRuntime/Private/LidarPointCloudLODManager.cpp
+++ b/Engine/Plugins/Enterprise/LidarPointCloud/Source/LidarPointCloudRuntime/Private/LidarPointCloudLODManager.cpp
@@ -183,19 +183,6 @@ int32 FLidarPointCloudTraversalOctree::GetVisibleNodes(TArray<FLidarPointCloudTr
 		const float BoundsScaleSq = SelectionParams.BoundsScale * SelectionParams.BoundsScale;
 		const float BaseLODImportance = FMath::Max(0.0f, CVarBaseLODImportance.GetValueOnAnyThread());
 
-		bool bStartClipped = false;
-		if (SelectionParams.ClippingVolumes)
-		{
-			for (const FLidarPointCloudClippingVolumeParams& Volume : *SelectionParams.ClippingVolumes)
-			{
-				if (Volume.Mode == ELidarClippingVolumeMode::ClipOutside)
-				{
-					bStartClipped = true;
-					break;
-				}
-			}
-		}
-
 		TQueue<FLidarPointCloudTraversalOctreeNode*> Nodes;
 		FLidarPointCloudTraversalOctreeNode* CurrentNode = nullptr;
 		Nodes.Enqueue(&Root);
@@ -215,37 +202,7 @@ int32 FLidarPointCloudTraversalOctree::GetVisibleNodes(TArray<FLidarPointCloudTr
 			{
 				continue;
 			}
-
-			const FBox NodeBounds(CurrentNode->Center - NodeExtent, CurrentNode->Center + NodeExtent);
-
-			// Check vs Clipping Volumes
-			bool bClip = bStartClipped;
-			if (SelectionParams.ClippingVolumes)
-			{
-				for (const FLidarPointCloudClippingVolumeParams& Volume : *SelectionParams.ClippingVolumes)
-				{
-					if (Volume.Mode == ELidarClippingVolumeMode::ClipOutside)
-					{
-						if (Volume.Bounds.Intersect(NodeBounds))
-						{
-							bClip = false;
-						}
-					}
-					else
-					{
-						if (Volume.Bounds.IsInside(NodeBounds))
-						{
-							bClip = true;
-						}
-					}
-				}
-
-				if (bClip)
-				{
-					continue;
-				}
-			}
-
+			
 			// Only process this node if it has any visible points - do not use continue; as the children may still contain visible points!
 			if (CurrentNode->DataNode->GetNumVisiblePoints() > 0 && CurrentNode->Depth >= SelectionParams.MinDepth)
 			{
