Removal of old per-asset editing for Lidar

diff --git a/Engine/Plugins/Enterprise/LidarPointCloud/Source/LidarPointCloudEditor/Private/LidarPointCloudEditor.cpp b/Engine/Plugins/Enterprise/LidarPointCloud/Source/LidarPointCloudEditor/Private/LidarPointCloudEditor.cpp
--- a/Engine/Plugins/Enterprise/LidarPointCloud/Source/LidarPointCloudEditor/Private/LidarPointCloudEditor.cpp
+++ b/Engine/Plugins/Enterprise/LidarPointCloud/Source/LidarPointCloudEditor/Private/LidarPointCloudEditor.cpp
@@ -27,85 +27,7 @@
 
 const FName PointCloudEditorAppName = FName(TEXT("LidarPointCloudEditorApp"));
 
-const FName FLidarPointCloudEditor::DetailsTabId(TEXT("Details"));
-const FName FLidarPointCloudEditor::ViewportTabId(TEXT("Viewport"));
-
-void SetSelectionMethod(FLidarPointCloudEditor* Editor, ELidarPointCloudSelectionMethod NewSelectionMethod)
-{
-	if (SLidarPointCloudEditorViewport* ViewportPtr = Editor->GetViewport().Get())
-	{
-		if (FLidarPointCloudEditorViewportClient* ViewportClientPtr = ViewportPtr->GetEditorViewportClient().Get())
-		{
-			ViewportClientPtr->SetSelectionMethod(NewSelectionMethod);
-		}
-	}
-
-	Editor->RegenerateMenusAndToolbars();
-}
-
-FName GetSelectionIcon(FLidarPointCloudEditor* Editor)
-{
-	if (SLidarPointCloudEditorViewport* ViewportPtr = Editor->GetViewport().Get())
-	{
-		if (FLidarPointCloudEditorViewportClient* ViewportClientPtr = ViewportPtr->GetEditorViewportClient().Get())
-		{
-			switch (ViewportClientPtr->GetSelectionMethod())
-			{
-			case ELidarPointCloudSelectionMethod::Box:
-				return "LidarPointCloudEditor.BoxSelection";
-
-			case ELidarPointCloudSelectionMethod::Polygonal:
-				return "LidarPointCloudEditor.PolygonalSelection";
-
-			case ELidarPointCloudSelectionMethod::Lasso:
-				return "LidarPointCloudEditor.LassoSelection";
-
-			case ELidarPointCloudSelectionMethod::Paint:
-				return "LidarPointCloudEditor.PaintSelection";
-			}
-		}
-	}
-
-	return "";
-}
-
-/////////////////////////////////////////////////////
-// SPointCloudPropertiesTabBody
-
-class SPointCloudPropertiesTabBody : public SSingleObjectDetailsPanel
-{
-public:
-	SLATE_BEGIN_ARGS(SPointCloudPropertiesTabBody) {}
-	SLATE_END_ARGS()
-
-private:
-	// Pointer back to owning sprite editor instance (the keeper of state)
-	TWeakPtr<class FLidarPointCloudEditor> PointCloudEditorPtr;
-public:
-	void Construct(const FArguments& InArgs, TSharedPtr<FLidarPointCloudEditor> InPointCloudEditor)
-	{
-		PointCloudEditorPtr = InPointCloudEditor;
-
-		SSingleObjectDetailsPanel::Construct(SSingleObjectDetailsPanel::FArguments().HostCommandList(InPointCloudEditor->GetToolkitCommands()).HostTabManager(InPointCloudEditor->GetTabManager()), /*bAutomaticallyObserveViaGetObjectToObserve=*/ true, /*bAllowSearch=*/ true);
-	}
-
-	// SSingleObjectDetailsPanel interface
-	virtual UObject* GetObjectToObserve() const override
-	{
-		return PointCloudEditorPtr.Pin()->GetPointCloudBeingEdited();
-	}
-
-	virtual TSharedRef<SWidget> PopulateSlot(TSharedRef<SWidget> PropertyEditorWidget) override
-	{
-		return SNew(SVerticalBox)
-			+ SVerticalBox::Slot()
-			.FillHeight(1)
-			[
-				PropertyEditorWidget
-			];
-	}
-	// End of SSingleObjectDetailsPanel interface
-};
+const FName ViewportTabId(TEXT("Viewport"));
 
 //////////////////////////////////////////////////////////////////////////
 // FLidarPointCloudEditor
@@ -124,34 +46,37 @@ FLidarPointCloudEditor::~FLidarPointCloudEditor()
 		PointCloudBeingEdited->OnPointCloudRebuilt().RemoveAll(this);
 		PointCloudBeingEdited->OnPreSaveCleanup().RemoveAll(this);
 	}
-
-	DeselectPoints();
 }
 
 void FLidarPointCloudEditor::RegisterTabSpawners(const TSharedRef<class FTabManager>& InTabManager)
 {
 	WorkspaceMenuCategory = InTabManager->AddLocalWorkspaceMenuCategory(LOCTEXT("WorkspaceMenu_LidarPointCloudEditor", "LiDAR Point Cloud Editor"));
-	TSharedRef<FWorkspaceItem> WorkspaceMenuCategoryRef = WorkspaceMenuCategory.ToSharedRef();
-
+	
 	FAssetEditorToolkit::RegisterTabSpawners(InTabManager);
+	
+	InTabManager->RegisterTabSpawner(ViewportTabId, FOnSpawnTab::CreateLambda([this](const FSpawnTabArgs& Args)
+	{
+		TSharedRef<SDockTab> SpawnedTab =
+			SNew(SDockTab)
+			.Label(LOCTEXT("ViewportTab_Title", "Viewport"))
+			[
+				Viewport.ToSharedRef()
+			];
 
-	InTabManager->RegisterTabSpawner(FLidarPointCloudEditor::DetailsTabId, FOnSpawnTab::CreateSP(this, &FLidarPointCloudEditor::SpawnTab_Details))
-		.SetDisplayName(LOCTEXT("DetailsTabLabel", "Details"))
-		.SetGroup(WorkspaceMenuCategoryRef)
-		.SetIcon(FSlateIcon(FAppStyle::GetAppStyleSetName(), "LevelEditor.Tabs.Details"));
+		Viewport->SetParentTab(SpawnedTab);
 
-	InTabManager->RegisterTabSpawner(FLidarPointCloudEditor::ViewportTabId, FOnSpawnTab::CreateSP(this, &FLidarPointCloudEditor::SpawnTab_Viewport))
-		.SetDisplayName(LOCTEXT("ViewportTabLabel", "Viewport"))
-		.SetGroup(WorkspaceMenuCategoryRef)
-		.SetIcon(FSlateIcon(FAppStyle::GetAppStyleSetName(), "LevelEditor.Tabs.Viewports"));
+		return SpawnedTab;
+	}))
+	.SetDisplayName(LOCTEXT("ViewportTabLabel", "Viewport"))
+	.SetGroup(WorkspaceMenuCategory.ToSharedRef())
+	.SetIcon(FSlateIcon(FAppStyle::GetAppStyleSetName(), "LevelEditor.Tabs.Viewports"));
 }
 
 void FLidarPointCloudEditor::UnregisterTabSpawners(const TSharedRef<class FTabManager>& InTabManager)
 {
 	FAssetEditorToolkit::UnregisterTabSpawners(InTabManager);
 
-	InTabManager->UnregisterTabSpawner(FLidarPointCloudEditor::DetailsTabId);
-	InTabManager->UnregisterTabSpawner(FLidarPointCloudEditor::ViewportTabId);
+	InTabManager->UnregisterTabSpawner(ViewportTabId);
 }
 
 FName FLidarPointCloudEditor::GetToolkitFName() const { return FName("LidarPointCloudEditor"); }
@@ -182,10 +107,6 @@ void FLidarPointCloudEditor::InitPointCloudEditor(const EToolkitMode::Type Mode,
 {
 	PointCloudBeingEdited = InitPointCloud;
 
-	// Register for rebuilding events
-	PointCloudBeingEdited->OnPointCloudRebuilt().AddSP(this, &FLidarPointCloudEditor::OnPointCloudRebuilt);
-	PointCloudBeingEdited->OnPreSaveCleanup().AddSP(this, &FLidarPointCloudEditor::OnPreSaveCleanup);
-
 	TSharedPtr<FLidarPointCloudEditor> PointCloudEditor = SharedThis(this);
 
 	Viewport = SNew(SLidarPointCloudEditorViewport).PointCloudEditor(SharedThis(this)).ObjectToEdit(PointCloudBeingEdited);
@@ -204,14 +125,7 @@ void FLidarPointCloudEditor::InitPointCloudEditor(const EToolkitMode::Type Mode,
 				(
 					FTabManager::NewStack()
 					->SetSizeCoefficient(0.6f)
-					->AddTab(FLidarPointCloudEditor::ViewportTabId, ETabState::OpenedTab)
-					->SetHideTabWell(true)
-				)
-				->Split
-				(
-					FTabManager::NewStack()
-					->SetSizeCoefficient(0.4f)
-					->AddTab(FLidarPointCloudEditor::DetailsTabId, ETabState::OpenedTab)
+					->AddTab(ViewportTabId, ETabState::OpenedTab)
 					->SetHideTabWell(true)
 				)
 			)
@@ -222,370 +136,9 @@ void FLidarPointCloudEditor::InitPointCloudEditor(const EToolkitMode::Type Mode,
 
 	ExtendToolBar();
 
-	BindEditorCommands();
-
 	RegenerateMenusAndToolbars();
 }
 
-void FLidarPointCloudEditor::SelectPointsByConvexVolume(const FConvexVolume& ConvexVolume, bool bAdditive)
-{
-	if (bAdditive && SelectedPoints.Num() > 0)
-	{
-		TArray64<FLidarPointCloudPoint*> NewSelectedPoints;
-		PointCloudBeingEdited->GetPointsInConvexVolume(NewSelectedPoints, ConvexVolume, true);
-
-		// Merge selections and set selection flag for new selection of points
-		for (FLidarPointCloudPoint** Data = NewSelectedPoints.GetData(), **DataEnd = Data + NewSelectedPoints.Num(); Data != DataEnd; ++Data)
-		{
-			if (!(*Data)->bSelected)
-			{
-				(*Data)->bSelected = true;
-				SelectedPoints.Add(*Data);
-			}
-		}
-	}
-	else
-	{
-		// Clear any existing selection
-		DeselectPoints();
-
-		PointCloudBeingEdited->GetPointsInConvexVolume(SelectedPoints, ConvexVolume, true);
-
-		// Set selection flag for new selection of points
-		for (FLidarPointCloudPoint** Data = SelectedPoints.GetData(), **DataEnd = Data + SelectedPoints.Num(); Data != DataEnd; ++Data)
-		{
-			(*Data)->bSelected = true;
-		}
-	}	
-
-	PointCloudBeingEdited->Octree.MarkRenderDataInConvexVolumeDirty(ConvexVolume);
-}
-
-void FLidarPointCloudEditor::DeselectPointsByConvexVolume(const FConvexVolume& ConvexVolume)
-{
-	TArray64<FLidarPointCloudPoint*> NewSelectedPoints;
-	PointCloudBeingEdited->GetPointsInConvexVolume(NewSelectedPoints, ConvexVolume, true);
-
-	// Unset selection flag for selection of points
-	for (FLidarPointCloudPoint** Data = NewSelectedPoints.GetData(), **DataEnd = Data + NewSelectedPoints.Num(); Data != DataEnd; ++Data)
-	{
-		(*Data)->bSelected = false;
-	}
-
-	// Remove deselected points from the selection list
-	FLidarPointCloudPoint** Data = SelectedPoints.GetData();
-	for (int32 i = 0; i < SelectedPoints.Num(); ++Data, ++i)
-	{
-		if (!(*Data)->bSelected)
-		{
-			SelectedPoints.RemoveAtSwap(i--, 1, false);
-			--Data;
-		}
-	}
-
-	SelectedPoints.Shrink();
-
-	PointCloudBeingEdited->Octree.MarkRenderDataInConvexVolumeDirty(ConvexVolume);
-}
-
-void FLidarPointCloudEditor::SelectPointsBySphere(const FSphere& Sphere)
-{
-	TArray64<FLidarPointCloudPoint*> NewSelectedPoints;
-	PointCloudBeingEdited->GetPointsInSphere(NewSelectedPoints, Sphere, true);
-
-	// Reserve space
-	const int64 NeededSlack = FMath::Max(0LL, NewSelectedPoints.Num() - SelectedPoints.GetSlack());
-	if (NeededSlack > 0)
-	{
-		SelectedPoints.Reserve(NeededSlack);
-	}
-
-	// Merge selections and set selection flag for new selection of points
-	for (FLidarPointCloudPoint** Data = NewSelectedPoints.GetData(), **DataEnd = Data + NewSelectedPoints.Num(); Data != DataEnd; ++Data)
-	{
-		if (!(*Data)->bSelected)
-		{
-			(*Data)->bSelected = true;
-			SelectedPoints.Add(*Data);
-		}
-	}
-
-	PointCloudBeingEdited->Octree.MarkRenderDataInSphereDirty(Sphere);
-}
-
-void FLidarPointCloudEditor::DeselectPointsBySphere(const FSphere& Sphere)
-{
-	TArray64<FLidarPointCloudPoint*> NewSelectedPoints;
-	PointCloudBeingEdited->GetPointsInSphere(NewSelectedPoints, Sphere, true);
-
-	// Unset selection flag for selection of points
-	for (FLidarPointCloudPoint** Data = NewSelectedPoints.GetData(), **DataEnd = Data + NewSelectedPoints.Num(); Data != DataEnd; ++Data)
-	{
-		(*Data)->bSelected = false;
-	}
-
-	// Remove deselected points from the selection list
-	FLidarPointCloudPoint** Data = SelectedPoints.GetData();
-	for (int32 i = 0; i < SelectedPoints.Num(); ++Data, ++i)
-	{
-		if (!(*Data)->bSelected)
-		{
-			SelectedPoints.RemoveAtSwap(i--, 1, false);
-			--Data;
-		}
-	}
-
-	PointCloudBeingEdited->Octree.MarkRenderDataInSphereDirty(Sphere);
-}
-
-void FLidarPointCloudEditor::DeselectPoints()
-{
-	for (FLidarPointCloudPoint** Data = SelectedPoints.GetData(), **DataEnd = Data + SelectedPoints.Num(); Data != DataEnd; ++Data)
-	{
-		(*Data)->bSelected = false;
-	}
-
-	SelectedPoints.Empty();
-
-	PointCloudBeingEdited->Octree.MarkRenderDataDirty();
-}
-
-void FLidarPointCloudEditor::InvertSelection()
-{
-	SelectedPoints.Empty(PointCloudBeingEdited->GetNumVisiblePoints() - SelectedPoints.Num());
-
-	PointCloudBeingEdited->ExecuteActionOnAllPoints([this](FLidarPointCloudPoint* Point){
-		if (Point->bSelected)
-		{
-			Point->bSelected = false;
-		}
-		else
-		{
-			Point->bSelected = true;
-			SelectedPoints.Add(Point);
-		}
-	}, true);
-
-	PointCloudBeingEdited->Octree.MarkRenderDataDirty();
-}
-
-void FLidarPointCloudEditor::DeletePoints()
-{
-	if (!ConfirmCollisionChange())
-	{
-		return;
-	}
-
-	// Permanently remove points
-	if (FMessageDialog::Open(EAppMsgType::Type::YesNo, FText::FromString("Warning: This operation cannot be reversed!\nAre you sure you want to continue?")) != EAppReturnType::Yes)
-	{
-		return;
-	}
-
-	PointCloudBeingEdited->RemovePoints(SelectedPoints);
-	PointCloudBeingEdited->RefreshRendering();
-	PointCloudBeingEdited->MarkPackageDirty();
-	SelectedPoints.Empty();
-
-	Viewport->RefreshViewport();
-}
-
-void FLidarPointCloudEditor::DeleteHiddenPoints()
-{
-	// Permanently remove points
-	if (FMessageDialog::Open(EAppMsgType::Type::YesNo, FText::FromString("Warning: This operation cannot be reversed!\nAre you sure you want to continue?")) != EAppReturnType::Yes)
-	{
-		return;
-	}
-
-	PointCloudBeingEdited->RemoveHiddenPoints();
-	DeselectPoints();
-	PointCloudBeingEdited->RefreshRendering();
-	PointCloudBeingEdited->MarkPackageDirty();
-
-	Viewport->RefreshViewport();
-}
-
-void FLidarPointCloudEditor::HidePoints()
-{
-	if (!ConfirmCollisionChange())
-	{
-		return;
-	}
-
-	// Hide points
-	for (FLidarPointCloudPoint** Data = SelectedPoints.GetData(), **DataEnd = Data + SelectedPoints.Num(); Data != DataEnd; ++Data)
-	{
-		(*Data)->bVisible = false;
-		(*Data)->bSelected = false;
-	}
-
-	PointCloudBeingEdited->MarkPointVisibilityDirty();
-
-	SelectedPoints.Empty();
-}
-
-void FLidarPointCloudEditor::UnhideAll()
-{
-	if (PointCloudBeingEdited)
-	{
-		if (!ConfirmCollisionChange())
-		{
-			return;
-		}
-
-		PointCloudBeingEdited->UnhideAll();
-	}
-}
-
-bool FLidarPointCloudEditor::ConfirmCollisionChange()
-{
-	if (PointCloudBeingEdited->HasCollisionData())
-	{
-		if (FMessageDialog::Open(EAppMsgType::Type::YesNo, FText::FromString("Performing this action will invalidate the collision data.\nAre you sure you want to continue?")) != EAppReturnType::Yes)
-		{
-			return false;
-		}
-
-		PointCloudBeingEdited->RemoveCollision();
-	}
-
-	return true;
-}
-
-TSharedRef<SWidget> FLidarPointCloudEditor::BuildPointCloudStatistics()
-{
-	return SNew(SHorizontalBox)
-		+ SHorizontalBox::Slot()
-		.FillWidth(0.5f)
-		[
-			SNew(SVerticalBox)
-			+ SVerticalBox::Slot()
-			.AutoHeight()
-			.VAlign(VAlign_Center)
-			.Padding(4.0f)
-			[
-				SNew(STextBlock).Text_Lambda([this] { return FText::Format(LOCTEXT("PointCount", "Point Count: {0}"), PointCloudBeingEdited ? PointCloudBeingEdited->GetNumPoints() : 0); })
-			]
-			+ SVerticalBox::Slot()
-			.AutoHeight()
-			.VAlign(VAlign_Center)
-			.Padding(4.0f)
-			[
-				SNew(STextBlock).Text_Lambda([this] { return FText::Format(LOCTEXT("NodeCount", "Node Count: {0}"), PointCloudBeingEdited ? PointCloudBeingEdited->GetNumNodes() : 0); })
-			]
-			+ SVerticalBox::Slot()
-			.AutoHeight()
-			.VAlign(VAlign_Center)
-			.Padding(4.0f)
-			[
-				SNew(STextBlock).Text_Lambda([this] { return FText::Format(LOCTEXT("Collider", "Collider: {0} poly"), PointCloudBeingEdited && PointCloudBeingEdited->HasCollisionData() ? PointCloudBeingEdited->Octree.GetCollisionData()->Indices.Num() : 0); })
-			]
-		]
-
-		+ SHorizontalBox::Slot()
-		.FillWidth(0.5f)
-		[
-			SNew(SVerticalBox)
-			+ SVerticalBox::Slot()
-			.AutoHeight()
-			.VAlign(VAlign_Center)
-			.Padding(4.0f)
-			[
-				SNew(STextBlock).Text_Lambda([this]
-				{
-					FVector BoundingSize = PointCloudBeingEdited ? PointCloudBeingEdited->GetBounds().GetSize() : FVector::ZeroVector;
-					return FText::Format(LOCTEXT("PCBounds", "Bounds: {0} x {1} x {2}"), FMath::CeilToInt(BoundingSize.X), FMath::CeilToInt(BoundingSize.Y), FMath::CeilToInt(BoundingSize.Z));
-				})
-			]
-			+ SVerticalBox::Slot()
-			.AutoHeight()
-			.VAlign(VAlign_Center)
-			.Padding(4.0f)
-			[
-				SNew(STextBlock).Text_Lambda([this] { return FText::Format(LOCTEXT("PCSize", "Data Size: {0} MB"), PointCloudBeingEdited ? PointCloudBeingEdited->GetDataSize() : 0); })
-			]
-		];
-}
-
-TSharedRef<SWidget> FLidarPointCloudEditor::GenerateNormalsMenuContent()
-{
-	FMenuBuilder MenuBuilder(true, Viewport->GetCommandList());
-
-	MenuBuilder.BeginSection("CalculateNormals", LOCTEXT("CalculateNormals", "Calculate Normals"));
-	{
-		MenuBuilder.AddMenuEntry(FLidarPointCloudEditorCommands::Get().CalculateNormals);
-		MenuBuilder.AddMenuEntry(FLidarPointCloudEditorCommands::Get().CalculateNormalsSelection);
-	}
-	MenuBuilder.EndSection();
-
-	return MenuBuilder.MakeWidget();
-}
-
-TSharedRef<SWidget> FLidarPointCloudEditor::GenerateExtractionMenuContent()
-{
-	FMenuBuilder MenuBuilder(true, Viewport->GetCommandList());
-
-	MenuBuilder.BeginSection("ExtractSelection", LOCTEXT("ExtractSelection", "Extract Selection"));
-	{
-		MenuBuilder.AddMenuEntry(FLidarPointCloudEditorCommands::Get().Extract);
-		MenuBuilder.AddMenuEntry(FLidarPointCloudEditorCommands::Get().ExtractCopy);
-	}
-	MenuBuilder.EndSection();
-
-	return MenuBuilder.MakeWidget();
-}
-
-TSharedRef<SWidget> FLidarPointCloudEditor::GenerateCollisionMenuContent()
-{
-	FMenuBuilder MenuBuilder(true, Viewport->GetCommandList());
-
-	MenuBuilder.BeginSection("Collision", LOCTEXT("Collision", "Collision"));
-	{
-		MenuBuilder.AddMenuEntry(FLidarPointCloudEditorCommands::Get().BuildCollision);
-		MenuBuilder.AddMenuEntry(FLidarPointCloudEditorCommands::Get().RemoveCollision);
-	}
-	MenuBuilder.EndSection();
-
-	return MenuBuilder.MakeWidget();
-}
-
-TSharedRef<SWidget> FLidarPointCloudEditor::GenerateDeleteMenuContent()
-{
-	FMenuBuilder MenuBuilder(true, Viewport->GetCommandList());
-
-	MenuBuilder.BeginSection("Delete", LOCTEXT("Delete", "Delete"));
-	{
-		MenuBuilder.AddMenuEntry(FLidarPointCloudEditorCommands::Get().DeleteSelected);
-		MenuBuilder.AddMenuEntry(FLidarPointCloudEditorCommands::Get().DeleteHidden, NAME_None, TAttribute<FText>(), TAttribute<FText>(), FSlateIcon(FLidarPointCloudStyle::GetStyleSetName(), "LidarPointCloudEditor.DeleteSelected"));
-	}
-	MenuBuilder.EndSection();
-
-	return MenuBuilder.MakeWidget();
-}
-
-TSharedRef<SWidget> FLidarPointCloudEditor::GenerateSelectionMenuContent()
-{
-	FMenuBuilder MenuBuilder(true, Viewport->GetCommandList());
-
-	MenuBuilder.BeginSection("SelectionMethod", LOCTEXT("SelectionMethod", "Selection Method"));
-	{
-		MenuBuilder.AddMenuEntry(FLidarPointCloudEditorCommands::Get().BoxSelection);
-		MenuBuilder.AddMenuEntry(FLidarPointCloudEditorCommands::Get().PolygonalSelection);
-		MenuBuilder.AddMenuEntry(FLidarPointCloudEditorCommands::Get().LassoSelection);
-		MenuBuilder.AddMenuEntry(FLidarPointCloudEditorCommands::Get().PaintSelection);
-	}
-	MenuBuilder.EndSection();
-
-	MenuBuilder.BeginSection("SelectionOptions", LOCTEXT("SelectionOptions", "Selection Options"));
-	{
-		MenuBuilder.AddMenuEntry(FLidarPointCloudEditorCommands::Get().InvertSelection);
-	}
-	MenuBuilder.EndSection();
-
-	return MenuBuilder.MakeWidget();
-}
-
 void FLidarPointCloudEditor::ExtendToolBar()
 {
 	struct Local
@@ -599,57 +152,6 @@ void FLidarPointCloudEditor::ExtendToolBar()
 				ToolbarBuilder.AddToolBarButton(Commands->ResetCamera, NAME_None, TAttribute<FText>(), TAttribute<FText>(), FSlateIcon(FAppStyle::GetAppStyleSetName(), "StaticMeshEditor.ResetCamera"));
 			}
 			ToolbarBuilder.EndSection();
-
-			ToolbarBuilder.BeginSection("LidarPointCloud");
-			{
-				ToolbarBuilder.AddToolBarButton(Commands->Center, NAME_None, TAttribute<FText>(), TAttribute<FText>(), FSlateIcon(FAppStyle::GetAppStyleSetName(), "AnimViewportMenu.TranslateMode"));
-				ToolbarBuilder.AddToolBarButton(Commands->Merge);
-				ToolbarBuilder.AddToolBarButton(Commands->Align);
-				ToolbarBuilder.AddComboButton(
-					FUIAction(FExecuteAction(), FCanExecuteAction()),
-					FOnGetContent::CreateSP(ThisEditor, &FLidarPointCloudEditor::GenerateCollisionMenuContent),
-					LOCTEXT("Collision_Label", "Collision"),
-					LOCTEXT("Collision_Tooltip", "Collision options"),
-					FSlateIcon(FLidarPointCloudStyle::GetStyleSetName(), "LidarPointCloudEditor.BuildCollision")
-				);
-			}
-			ToolbarBuilder.EndSection();
-
-			ToolbarBuilder.BeginSection("LidarPointCloudEdit");
-			{
-				ToolbarBuilder.AddToolBarButton(Commands->EditMode);
-				ToolbarBuilder.AddComboButton(
-					FUIAction(FExecuteAction(), FCanExecuteAction::CreateSP(ThisEditor, &FLidarPointCloudEditor::IsEditMode)),
-					FOnGetContent::CreateSP(ThisEditor, &FLidarPointCloudEditor::GenerateSelectionMenuContent),
-					LOCTEXT("Selection_Label", "Selection"),
-					LOCTEXT("Selection_Tooltip", "Selection options"),
-					FSlateIcon(FLidarPointCloudStyle::GetStyleSetName(), GetSelectionIcon(ThisEditor))
-				);			
-				ToolbarBuilder.AddToolBarButton(Commands->HideSelected);
-				ToolbarBuilder.AddToolBarButton(Commands->UnhideAll);
-				ToolbarBuilder.AddComboButton(
-					FUIAction(FExecuteAction(), FCanExecuteAction()),
-					FOnGetContent::CreateSP(ThisEditor, &FLidarPointCloudEditor::GenerateDeleteMenuContent),
-					LOCTEXT("Delete_Label", "Delete"),
-					LOCTEXT("Delete_Tooltip", "Point deletion options"),
-					FSlateIcon(FLidarPointCloudStyle::GetStyleSetName(), "LidarPointCloudEditor.DeleteSelected")
-				);
-				ToolbarBuilder.AddComboButton(
-					FUIAction(FExecuteAction(), FCanExecuteAction::CreateSP(ThisEditor, &FLidarPointCloudEditor::HasSelectedPoints)),
-					FOnGetContent::CreateSP(ThisEditor, &FLidarPointCloudEditor::GenerateExtractionMenuContent),
-					LOCTEXT("Extract_Label", "Extract"),
-					LOCTEXT("Extract_Tooltip", "Selection extraction options"),
-					FSlateIcon(FLidarPointCloudStyle::GetStyleSetName(), "LidarPointCloudEditor.Extract")
-				);
-				ToolbarBuilder.AddComboButton(
-					FUIAction(FExecuteAction(), FCanExecuteAction::CreateSP(ThisEditor, &FLidarPointCloudEditor::IsEditMode)),
-					FOnGetContent::CreateSP(ThisEditor, &FLidarPointCloudEditor::GenerateNormalsMenuContent),
-					LOCTEXT("Normals_Label", "Normals"),
-					LOCTEXT("Normals_Tooltip", "Normal Calculation options"),
-					FSlateIcon(FAppStyle::GetAppStyleSetName(), "AnimViewportMenu.SetShowNormals")
-				);
-			}
-			ToolbarBuilder.EndSection();
 		}
 	};
 
@@ -660,291 +162,4 @@ void FLidarPointCloudEditor::ExtendToolBar()
 	AddToolbarExtender(ILidarPointCloudEditorModule::Get().GetToolBarExtensibilityManager()->GetAllExtenders(GetToolkitCommands(), GetEditingObjects()));
 }
 
-void FLidarPointCloudEditor::BindEditorCommands()
-{
-	TSharedPtr<FUICommandList> CommandList = Viewport->GetCommandList();
-	
-	const FLidarPointCloudEditorCommands& Commands = FLidarPointCloudEditorCommands::Get();
-
-	CommandList->MapAction(Commands.Center, FExecuteAction::CreateSP(this, &FLidarPointCloudEditor::ToggleCenter), FCanExecuteAction(), FIsActionChecked::CreateSP(this, &FLidarPointCloudEditor::IsCentered));
-	CommandList->MapAction(Commands.BuildCollision, FExecuteAction::CreateSP(this, &FLidarPointCloudEditor::BuildCollision));
-	CommandList->MapAction(Commands.RemoveCollision, FExecuteAction::CreateSP(this, &FLidarPointCloudEditor::RemoveCollision), FCanExecuteAction::CreateSP(this, &FLidarPointCloudEditor::HasCollisionData));
-	CommandList->MapAction(Commands.EditMode, FExecuteAction::CreateSP(this, &FLidarPointCloudEditor::ToggleEditMode), FCanExecuteAction(), FIsActionChecked::CreateSP(this, &FLidarPointCloudEditor::IsEditMode));
-	CommandList->MapAction(Commands.BoxSelection, FExecuteAction::CreateLambda([this]{ SetSelectionMethod(this, ELidarPointCloudSelectionMethod::Box); }), FCanExecuteAction::CreateSP(this, &FLidarPointCloudEditor::IsEditMode));
-	CommandList->MapAction(Commands.PolygonalSelection, FExecuteAction::CreateLambda([this]{ SetSelectionMethod(this, ELidarPointCloudSelectionMethod::Polygonal); }), FCanExecuteAction::CreateSP(this, &FLidarPointCloudEditor::IsEditMode));
-	CommandList->MapAction(Commands.LassoSelection, FExecuteAction::CreateLambda([this] { SetSelectionMethod(this, ELidarPointCloudSelectionMethod::Lasso); }), FCanExecuteAction::CreateSP(this, &FLidarPointCloudEditor::IsEditMode));
-	CommandList->MapAction(Commands.PaintSelection, FExecuteAction::CreateLambda([this] { SetSelectionMethod(this, ELidarPointCloudSelectionMethod::Paint); }), FCanExecuteAction::CreateSP(this, &FLidarPointCloudEditor::IsEditMode));
-	CommandList->MapAction(Commands.InvertSelection, FExecuteAction::CreateSP(this, &FLidarPointCloudEditor::InvertSelection), FCanExecuteAction::CreateSP(this, &FLidarPointCloudEditor::HasSelectedPoints));
-	CommandList->MapAction(Commands.UnhideAll, FExecuteAction::CreateSP(this, &FLidarPointCloudEditor::UnhideAll));
-	CommandList->MapAction(Commands.HideSelected, FExecuteAction::CreateSP(this, &FLidarPointCloudEditor::HidePoints), FCanExecuteAction::CreateSP(this, &FLidarPointCloudEditor::HasSelectedPoints));
-	CommandList->MapAction(Commands.DeleteSelected, FExecuteAction::CreateSP(this, &FLidarPointCloudEditor::DeletePoints), FCanExecuteAction::CreateSP(this, &FLidarPointCloudEditor::HasSelectedPoints));
-	CommandList->MapAction(Commands.DeleteHidden, FExecuteAction::CreateSP(this, &FLidarPointCloudEditor::DeleteHiddenPoints));
-	CommandList->MapAction(Commands.Extract, FExecuteAction::CreateSP(this, &FLidarPointCloudEditor::Extract), FCanExecuteAction::CreateSP(this, &FLidarPointCloudEditor::HasSelectedPoints));
-	CommandList->MapAction(Commands.ExtractCopy, FExecuteAction::CreateSP(this, &FLidarPointCloudEditor::ExtractCopy), FCanExecuteAction::CreateSP(this, &FLidarPointCloudEditor::HasSelectedPoints));
-	CommandList->MapAction(Commands.Merge, FExecuteAction::CreateSP(this, &FLidarPointCloudEditor::Merge));
-	CommandList->MapAction(Commands.Align, FExecuteAction::CreateSP(this, &FLidarPointCloudEditor::Align));
-	CommandList->MapAction(Commands.CalculateNormals, FExecuteAction::CreateSP(this, &FLidarPointCloudEditor::CalculateNormals));
-	CommandList->MapAction(Commands.CalculateNormalsSelection, FExecuteAction::CreateSP(this, &FLidarPointCloudEditor::CalculateNormalsSelection));
-}
-
-TSharedRef<SDockTab> FLidarPointCloudEditor::SpawnTab_Details(const FSpawnTabArgs& Args)
-{
-	TSharedPtr<FLidarPointCloudEditor> PointCloudEditorPtr = SharedThis(this);
-
-	// Spawn the tab
-	return SNew(SDockTab)
-		.Label(LOCTEXT("DetailsTab_Title", "Details"))
-		[
-			SNew(SVerticalBox)
-				+ SVerticalBox::Slot()
-				.AutoHeight()
-				.Padding(2.0f)
-				[
-					SNew(SBorder)
-					[
-						BuildPointCloudStatistics()
-					]
-				]
-
-				+ SVerticalBox::Slot()
-				.FillHeight(1.0f)
-				.Padding(2.0f)
-				[
-					SNew(SBorder)
-					.Padding(4.0f)
-					[
-						SNew(SPointCloudPropertiesTabBody, PointCloudEditorPtr)
-					]
-				]
-		];
-}
-
-TSharedRef<SDockTab> FLidarPointCloudEditor::SpawnTab_Viewport(const FSpawnTabArgs& Args)
-{
-	check(Args.GetTabId() == ViewportTabId);
-
-	TSharedRef<SDockTab> SpawnedTab =
-		SNew(SDockTab)
-		.Label(LOCTEXT("ViewportTab_Title", "Viewport"))
-		[
-			Viewport.ToSharedRef()
-		];
-
-	Viewport->SetParentTab(SpawnedTab);
-
-	return SpawnedTab;
-}
-
-void FLidarPointCloudEditor::OnPointCloudRebuilt()
-{
-	// If the cloud asset has been rebuilt, invalidate the selection array
-	SelectedPoints.Empty();
-}
-
-void FLidarPointCloudEditor::OnPreSaveCleanup()
-{
-	// If the cloud asset is being saved, deselect all points
-	DeselectPoints();
-}
-
-void FLidarPointCloudEditor::Extract()
-{
-	// Skip, if no points are selected
-	if (SelectedPoints.Num() == 0)
-	{
-		return;
-	}
-
-	if (!ConfirmCollisionChange())
-	{
-		return;
-	}
-
-	ULidarPointCloud* NewPointCloud = FLidarPointCloudEditorHelper::CreateNewAsset();
-	if (NewPointCloud)
-	{
-		NewPointCloud->SetData(SelectedPoints);
-		NewPointCloud->OriginalCoordinates = PointCloudBeingEdited->OriginalCoordinates;
-		DeletePoints();
-	}
-}
-
-void FLidarPointCloudEditor::ExtractCopy()
-{
-	// Skip, if no points are selected
-	if (SelectedPoints.Num() == 0)
-	{
-		return;
-	}
-
-	ULidarPointCloud* NewPointCloud = FLidarPointCloudEditorHelper::CreateNewAsset();
-	if (NewPointCloud)
-	{
-		NewPointCloud->SetData(SelectedPoints);
-		NewPointCloud->OriginalCoordinates = PointCloudBeingEdited->OriginalCoordinates;
-		DeselectPoints();
-	}
-}
-
-void FLidarPointCloudEditor::ToggleCenter()
-{
-	if (!ConfirmCollisionChange())
-	{
-		return;
-	}
-
-	DeselectPoints();
-
-	if (IsCentered())
-	{
-		PointCloudBeingEdited->RestoreOriginalCoordinates();
-	}
-	else
-	{
-		PointCloudBeingEdited->CenterPoints();
-	}
-
-	if (Viewport.IsValid())
-	{
-		Viewport->ResetCamera();
-	}
-}
-
-void FLidarPointCloudEditor::ToggleEditMode()
-{
-	if (!bEditMode)
-	{
-		// Warn about loading the whole cloud
-		if (!PointCloudBeingEdited->IsFullyLoaded())
-		{
-			if (FMessageDialog::Open(EAppMsgType::Type::YesNo, FText::FromString("The asset needs to be fully loaded into memory to enable editing.\nThis may take a while, depending on the size of the asset.\nAre you sure you want to continue?")) != EAppReturnType::Yes)
-			{
-				return;
-			}
-		}
-
-		PointCloudBeingEdited->LoadAllNodes();
-	}
-
-	bEditMode = !bEditMode;
-
-	if (!bEditMode)
-	{
-		DeselectPoints();
-		Viewport->RefreshViewport();
-	}
-}
-
-bool FLidarPointCloudEditor::IsCentered() const
-{
-	return PointCloudBeingEdited && PointCloudBeingEdited->IsCentered();
-}
-
-void FLidarPointCloudEditor::Merge()
-{
-	TArray<FAssetData> SelectedAssets = SelectAssets(LOCTEXT("SelectMergeSource", "Select Merge Source"));
-	if (SelectedAssets.Num())
-	{
-		TArray<ULidarPointCloud*> Assets;
-
-		for (int32 i = 0; i < SelectedAssets.Num(); ++i)
-		{
-			ULidarPointCloud* Asset = Cast<ULidarPointCloud>(SelectedAssets[i].GetAsset());
-
-			// Skip invalid and self
-			if (!Asset || Asset == PointCloudBeingEdited)
-			{
-				continue;
-			}
-
-			Assets.Add(Asset);
-		}
-
-		FScopedSlowTask ProgressDialog(Assets.Num() + 2, LOCTEXT("Merge", "Merging Point Clouds..."));
-		ProgressDialog.MakeDialog();
-
-		PointCloudBeingEdited->Merge(Assets, [&ProgressDialog]() { ProgressDialog.EnterProgressFrame(1.f); });
-	}
-}
-
-void FLidarPointCloudEditor::BuildCollision()
-{
-	if (PointCloudBeingEdited)
-	{
-		PointCloudBeingEdited->BuildCollision();
-	}
-}
-
-void FLidarPointCloudEditor::RemoveCollision()
-{
-	if (PointCloudBeingEdited)
-	{
-		PointCloudBeingEdited->RemoveCollision();
-	}
-}
-
-void FLidarPointCloudEditor::Align()
-{
-	if (!ConfirmCollisionChange())
-	{
-		return;
-	}
-
-	TArray<FAssetData> SelectedAssets = SelectAssets(LOCTEXT("SelectAlignmentSources", "Select Alignment Sources"));
-	if (SelectedAssets.Num())
-	{
-		TArray<ULidarPointCloud*> Assets;
-		Assets.Add(PointCloudBeingEdited);
-
-		for (int32 i = 0; i < SelectedAssets.Num(); ++i)
-		{
-			ULidarPointCloud* Asset = Cast<ULidarPointCloud>(SelectedAssets[i].GetAsset());
-
-			// Skip invalid and self
-			if (!Asset || Asset == PointCloudBeingEdited)
-			{
-				continue;
-			}
-
-			Assets.Add(Asset);
-		}
-
-		FScopedSlowTask ProgressDialog(1, LOCTEXT("Align", "Aligning Point Clouds..."));
-		ProgressDialog.MakeDialog();
-		ProgressDialog.EnterProgressFrame(1.f);
-		ULidarPointCloud::AlignClouds(Assets);
-	}
-}
-
-void FLidarPointCloudEditor::CalculateNormals()
-{
-	PointCloudBeingEdited->CalculateNormals(nullptr, nullptr);
-}
-
-void FLidarPointCloudEditor::CalculateNormalsSelection()
-{
-	if (SelectedPoints.Num() > 0)
-	{
-		PointCloudBeingEdited->CalculateNormals(&SelectedPoints, nullptr);
-	}
-}
-
-bool FLidarPointCloudEditor::HasCollisionData() const
-{
-	return PointCloudBeingEdited && PointCloudBeingEdited->HasCollisionData();
-}
-
-TArray<FAssetData> FLidarPointCloudEditor::SelectAssets(const FText& Title)
-{
-	// Initialize SaveAssetDialog config
-	FOpenAssetDialogConfig OpenAssetDialogConfig;
-	OpenAssetDialogConfig.DialogTitleOverride = Title;
-	OpenAssetDialogConfig.DefaultPath = "/Game";
-	OpenAssetDialogConfig.bAllowMultipleSelection = true;
-	OpenAssetDialogConfig.AssetClassNames.Add(ULidarPointCloud::StaticClass()->GetClassPathName());
-
-	FContentBrowserModule& ContentBrowserModule = FModuleManager::LoadModuleChecked<FContentBrowserModule>("ContentBrowser");
-	return ContentBrowserModule.Get().CreateModalOpenAssetDialog(OpenAssetDialogConfig);
-}
-
 #undef LOCTEXT_NAMESPACE
\ No newline at end of file
diff --git a/Engine/Plugins/Enterprise/LidarPointCloud/Source/LidarPointCloudEditor/Private/LidarPointCloudEditor.h b/Engine/Plugins/Enterprise/LidarPointCloud/Source/LidarPointCloudEditor/Private/LidarPointCloudEditor.h
--- a/Engine/Plugins/Enterprise/LidarPointCloud/Source/LidarPointCloudEditor/Private/LidarPointCloudEditor.h
+++ b/Engine/Plugins/Enterprise/LidarPointCloud/Source/LidarPointCloudEditor/Private/LidarPointCloudEditor.h
@@ -61,65 +61,9 @@ public:
 
 	ULidarPointCloud* GetPointCloudBeingEdited() const { return PointCloudBeingEdited; }
 
-	void SelectPointsByConvexVolume(const FConvexVolume& ConvexVolume, bool bAdditive);
-	void DeselectPointsByConvexVolume(const FConvexVolume& ConvexVolume);
-	void SelectPointsBySphere(const FSphere& Sphere);
-	void DeselectPointsBySphere(const FSphere& Sphere);
-	void DeselectPoints();
-	void InvertSelection();
-	void DeletePoints();
-	void DeleteHiddenPoints();
-	void HidePoints();
-	void UnhideAll();
-	TArray64<FLidarPointCloudPoint*>& GetSelectedPoints() { return SelectedPoints; }
-
-	bool HasSelectedPoints() const { return SelectedPoints.Num() > 0; }
-
-	bool IsEditMode() const { return bEditMode; }
-
 	TSharedPtr<SLidarPointCloudEditorViewport> GetViewport() { return Viewport; }
 
 private:
-	bool ConfirmCollisionChange();
-
-	TSharedRef<SWidget> BuildPointCloudStatistics();
-
 	/** Builds the Point Cloud Editor toolbar. */
 	void ExtendToolBar();
-
-	void BindEditorCommands();
-
-	TSharedRef<SDockTab> SpawnTab_Details(const FSpawnTabArgs& Args);
-	TSharedRef<SDockTab> SpawnTab_Viewport(const FSpawnTabArgs& Args);
-
-	/** Called when the owned cloud is rebuilt */
-	void OnPointCloudRebuilt();
-
-	void OnPreSaveCleanup();
-
-	void Extract();
-	void ExtractCopy();
-	bool IsCentered() const;
-	void ToggleCenter();
-	void ToggleEditMode();
-	void Merge();
-	void BuildCollision();
-	void RemoveCollision();
-	void Align();
-	void CalculateNormals();
-	void CalculateNormalsSelection();
-	bool HasCollisionData() const;
-
-	TSharedRef<SWidget> GenerateNormalsMenuContent();
-	TSharedRef<SWidget> GenerateExtractionMenuContent();
-	TSharedRef<SWidget> GenerateCollisionMenuContent();
-	TSharedRef<SWidget> GenerateDeleteMenuContent();
-	TSharedRef<SWidget> GenerateSelectionMenuContent();
-
-	TArray<FAssetData> SelectAssets(const FText& Title);
-
-private:
-	/**	The tab ids for all the tabs used */
-	static const FName DetailsTabId;
-	static const FName ViewportTabId;
 };
\ No newline at end of file
diff --git a/Engine/Plugins/Enterprise/LidarPointCloud/Source/LidarPointCloudEditor/Private/LidarPointCloudEditorCommands.cpp b/Engine/Plugins/Enterprise/LidarPointCloud/Source/LidarPointCloudEditor/Private/LidarPointCloudEditorCommands.cpp
--- a/Engine/Plugins/Enterprise/LidarPointCloud/Source/LidarPointCloudEditor/Private/LidarPointCloudEditorCommands.cpp
+++ b/Engine/Plugins/Enterprise/LidarPointCloud/Source/LidarPointCloudEditor/Private/LidarPointCloudEditorCommands.cpp
@@ -17,24 +17,6 @@ void FLidarPointCloudEditorCommands::RegisterCommands()
 	UI_COMMAND(SetShowNodes, "Nodes", "Toggles display of the nodes of the point cloud.", EUserInterfaceActionType::ToggleButton, FInputChord());
 	UI_COMMAND(ResetCamera, "Reset Camera", "Resets the camera to focus on the point cloud.", EUserInterfaceActionType::Button, FInputChord());
 	UI_COMMAND(Center, "Center", "Enable, to center the point cloud asset\nDisable, to use original coordinates.", EUserInterfaceActionType::ToggleButton, FInputChord());
-	UI_COMMAND(BuildCollision, "Build Collision", "Builds collision for this point cloud.", EUserInterfaceActionType::Button, FInputChord());
-	UI_COMMAND(RemoveCollision, "Remove Collision", "Removes collision from this point cloud.", EUserInterfaceActionType::Button, FInputChord());
-	UI_COMMAND(EditMode, "Edit Mode", "Enables editing of the point cloud.", EUserInterfaceActionType::ToggleButton, FInputChord());
-	UI_COMMAND(BoxSelection, "Box Selection", "Uses box to select points.", EUserInterfaceActionType::Button, FInputChord());
-	UI_COMMAND(PolygonalSelection, "Polygonal Selection", "Uses custom polygon to select points.", EUserInterfaceActionType::Button, FInputChord());
-	UI_COMMAND(LassoSelection, "Lasso Selection", "Uses custom drawn shape to select points.", EUserInterfaceActionType::Button, FInputChord());
-	UI_COMMAND(PaintSelection, "Paint Selection", "Uses adjustable paint brush to select points.", EUserInterfaceActionType::Button, FInputChord());
-	UI_COMMAND(InvertSelection, "Invert Selection", "Inverts point selection.", EUserInterfaceActionType::Button, FInputChord());
-	UI_COMMAND(HideSelected, "Hide Selected", "Hide selected points.", EUserInterfaceActionType::Button, FInputChord());
-	UI_COMMAND(UnhideAll, "Unhide All", "Resets the visibility of all points in the point cloud.", EUserInterfaceActionType::Button, FInputChord());
-	UI_COMMAND(DeleteSelected, "Delete Selected", "Permanently remove selected points from the point cloud.", EUserInterfaceActionType::Button, FInputChord());
-	UI_COMMAND(DeleteHidden, "Delete Hidden", "Permanently remove hidden points from the point cloud.", EUserInterfaceActionType::Button, FInputChord());
-	UI_COMMAND(CalculateNormals, "Calculate Normals", "Calculates normals for the point cloud.", EUserInterfaceActionType::Button, FInputChord());
-	UI_COMMAND(CalculateNormalsSelection, "Calculate Normals (Selection)", "Calculates normals for the selected points.", EUserInterfaceActionType::Button, FInputChord());
-	UI_COMMAND(Extract, "Extract", "Extracts the selected points as a separate point cloud asset.", EUserInterfaceActionType::Button, FInputChord());
-	UI_COMMAND(ExtractCopy, "Extract as Copy", "Duplicates the selected points as a separate point cloud asset.", EUserInterfaceActionType::Button, FInputChord());
-	UI_COMMAND(Merge, "Merge", "Merges selected point cloud assets with this one.", EUserInterfaceActionType::Button, FInputChord());
-	UI_COMMAND(Align, "Align", "Aligns selected point cloud assets with this one while retaining overall centering.", EUserInterfaceActionType::Button, FInputChord());
 
 	// TOOLKIT COMMANDS
 
diff --git a/Engine/Plugins/Enterprise/LidarPointCloud/Source/LidarPointCloudEditor/Private/LidarPointCloudEditorCommands.h b/Engine/Plugins/Enterprise/LidarPointCloud/Source/LidarPointCloudEditor/Private/LidarPointCloudEditorCommands.h
--- a/Engine/Plugins/Enterprise/LidarPointCloud/Source/LidarPointCloudEditor/Private/LidarPointCloudEditorCommands.h
+++ b/Engine/Plugins/Enterprise/LidarPointCloud/Source/LidarPointCloudEditor/Private/LidarPointCloudEditorCommands.h
@@ -36,24 +36,6 @@ public:
 	TSharedPtr<FUICommandInfo> SetShowNodes;
 	TSharedPtr<FUICommandInfo> ResetCamera;
 	TSharedPtr<FUICommandInfo> Center;
-	TSharedPtr<FUICommandInfo> BuildCollision;
-	TSharedPtr<FUICommandInfo> RemoveCollision;
-	TSharedPtr<FUICommandInfo> EditMode;
-	TSharedPtr<FUICommandInfo> BoxSelection;
-	TSharedPtr<FUICommandInfo> PolygonalSelection;
-	TSharedPtr<FUICommandInfo> LassoSelection;
-	TSharedPtr<FUICommandInfo> PaintSelection;
-	TSharedPtr<FUICommandInfo> InvertSelection;
-	TSharedPtr<FUICommandInfo> HideSelected;
-	TSharedPtr<FUICommandInfo> UnhideAll;
-	TSharedPtr<FUICommandInfo> DeleteSelected;
-	TSharedPtr<FUICommandInfo> DeleteHidden;
-	TSharedPtr<FUICommandInfo> CalculateNormals;
-	TSharedPtr<FUICommandInfo> CalculateNormalsSelection;
-	TSharedPtr<FUICommandInfo> Extract;
-	TSharedPtr<FUICommandInfo> ExtractCopy;
-	TSharedPtr<FUICommandInfo> Merge;
-	TSharedPtr<FUICommandInfo> Align;
 
 	// TOOLKIT COMMANDS
 
diff --git a/Engine/Plugins/Enterprise/LidarPointCloud/Source/LidarPointCloudEditor/Private/LidarPointCloudEditorViewport.cpp b/Engine/Plugins/Enterprise/LidarPointCloud/Source/LidarPointCloudEditor/Private/LidarPointCloudEditorViewport.cpp
--- a/Engine/Plugins/Enterprise/LidarPointCloud/Source/LidarPointCloudEditor/Private/LidarPointCloudEditorViewport.cpp
+++ b/Engine/Plugins/Enterprise/LidarPointCloud/Source/LidarPointCloudEditor/Private/LidarPointCloudEditorViewport.cpp
@@ -10,7 +10,6 @@
 #include "Styling/AppStyle.h"
 #include "ComponentReregisterContext.h"
 #include "Widgets/Docking/SDockTab.h"
-#include "Components/StaticMeshComponent.h"
 
 ///////////////////////////////////////////////////////////
 // SPointCloudEditorViewportToolbar
@@ -73,18 +72,7 @@ void SLidarPointCloudEditorViewport::Construct(const FArguments& InArgs)
 	PreviewScene.AddComponent(PreviewCloudComponent, FTransform::Identity);
 
 	SetPreviewCloud(PointCloud);
-
-	// Add Paint Brush
-	static UStaticMesh* PaintBrushMesh = Cast<UStaticMesh>(StaticLoadObject(UStaticMesh::StaticClass(), nullptr, TEXT("/Engine/BasicShapes/Sphere.Sphere")));
-	static UMaterialInterface* PaintBrushMat = Cast<UMaterialInterface>(StaticLoadObject(UMaterialInterface::StaticClass(), nullptr, TEXT("/LidarPointCloud/Materials/M_LidarSelectionBrush.M_LidarSelectionBrush")));
-
-	PaintBrush = NewObject<UStaticMeshComponent>();
-	PaintBrush->SetMobility(EComponentMobility::Movable);
-	PaintBrush->SetVisibility(false);
-	PaintBrush->SetStaticMesh(PaintBrushMesh);
-	PaintBrush->SetMaterial(0, PaintBrushMat);
-	PreviewScene.AddComponent(PaintBrush, FTransform::Identity);
-
+	
 	ViewportOverlay->AddSlot()
 		.VAlign(VAlign_Top)
 		.HAlign(HAlign_Left)
@@ -123,7 +111,6 @@ SLidarPointCloudEditorViewport::~SLidarPointCloudEditorViewport()
 void SLidarPointCloudEditorViewport::AddReferencedObjects(FReferenceCollector& Collector)
 {
 	Collector.AddReferencedObject(PreviewCloudComponent);
-	Collector.AddReferencedObject(PaintBrush);
 	Collector.AddReferencedObject(PointCloud);
 }
 
@@ -232,6 +219,5 @@ void SLidarPointCloudEditorViewport::OnFocusViewportToSelection()
 	if (PreviewCloudComponent)
 	{
 		EditorViewportClient->FocusViewportOnBox(PreviewCloudComponent->Bounds.GetBox());
-		return;
 	}
 }
diff --git a/Engine/Plugins/Enterprise/LidarPointCloud/Source/LidarPointCloudEditor/Private/LidarPointCloudEditorViewport.h b/Engine/Plugins/Enterprise/LidarPointCloud/Source/LidarPointCloudEditor/Private/LidarPointCloudEditorViewport.h
--- a/Engine/Plugins/Enterprise/LidarPointCloud/Source/LidarPointCloudEditor/Private/LidarPointCloudEditorViewport.h
+++ b/Engine/Plugins/Enterprise/LidarPointCloud/Source/LidarPointCloudEditor/Private/LidarPointCloudEditorViewport.h
@@ -3,7 +3,6 @@
 #pragma once
 
 #include "CoreMinimal.h"
-#include "SlateFwd.h"
 #include "UObject/GCObject.h"
 #include "EditorViewportClient.h"
 #include "AdvancedPreviewScene.h"
@@ -15,7 +14,6 @@ class FLidarPointCloudEditorViewportClient;
 class SVerticalBox;
 class ULidarPointCloud;
 class ULidarPointCloudComponent;
-class UStaticMeshComponent;
 
 /**
  * PointCloud Editor Preview viewport widget
@@ -46,9 +44,7 @@ public:
 
 	/** Component for the preview point cloud. */
 	ULidarPointCloudComponent* PreviewCloudComponent;
-
-	UStaticMeshComponent* PaintBrush;
-
+	
 	/**
 	 *	Sets up the point cloud that the Point Cloud editor is viewing.
 	 *
@@ -91,7 +87,7 @@ protected:
 
 private:
 	/** Determines the visibility of the viewport. */
-	bool IsVisible() const override;
+	virtual bool IsVisible() const override;
 
 private:
 	/** The parent tab where this viewport resides */
diff --git a/Engine/Plugins/Enterprise/LidarPointCloud/Source/LidarPointCloudEditor/Private/LidarPointCloudEditorViewportClient.cpp b/Engine/Plugins/Enterprise/LidarPointCloud/Source/LidarPointCloudEditor/Private/LidarPointCloudEditorViewportClient.cpp
--- a/Engine/Plugins/Enterprise/LidarPointCloud/Source/LidarPointCloudEditor/Private/LidarPointCloudEditorViewportClient.cpp
+++ b/Engine/Plugins/Enterprise/LidarPointCloud/Source/LidarPointCloudEditor/Private/LidarPointCloudEditorViewportClient.cpp
@@ -24,16 +24,6 @@ namespace UE {
 		namespace Private {
 			namespace Editor
 			{
-				// Distance Square between the first and last points of the polygonal selection, where the shape will be considered as closed
-				constexpr int32 PolySnapDistanceSq = 40;
-
-				// Affects the frequency of new point injections when drawing lasso-based shapes
-				constexpr int32 LassoSpacingSq = 400;
-
-				// Affects the max depth delta when painting. Prevents the brush from "falling through" the gaps.
-				constexpr float PaintMaxDeviation = 0.15f;
-
-				// Defaults for the common draw helper
 				constexpr float GridSize = 2048.0f;
 				constexpr int32 CellSize = 16;
 			}
@@ -41,53 +31,11 @@ namespace UE {
 	}
 }
 
-TArray<FVector2D> ToVectorArray(const TArray<FIntPoint>& Points)
-{
-	TArray<FVector2D> VectorPoints;
-	VectorPoints.Reserve(Points.Num());
-	for (const FIntPoint& Point : Points)
-	{
-		VectorPoints.Add(Point);
-	}
-	return VectorPoints;
-}
-
-// Copied from GeomTools.cpp and converted to work with FIntPoint
-bool IsPolygonConvex(const TArray<FIntPoint>& Points)
-{
-	const int32 PointCount = Points.Num();
-	int32 Sign = 0;
-	for (int32 PointIndex = 0; PointIndex < PointCount; ++PointIndex)
-	{
-		const FIntPoint& A = Points[PointIndex];
-		const FIntPoint& B = Points[(PointIndex + 1) % PointCount];
-		const FIntPoint& C = Points[(PointIndex + 2) % PointCount];
-		int32 Det = (B.X - A.X) * (C.Y - B.Y) - (B.Y - A.Y) * (C.X - B.X);
-		int32 DetSign = FMath::Sign(Det);
-		if (DetSign != 0)
-		{
-			if (Sign == 0)
-			{
-				Sign = DetSign;
-			}
-			else if (Sign != DetSign)
-			{
-				return false;
-			}
-		}
-	}
-
-	return true;
-}
-
 FLidarPointCloudEditorViewportClient::FLidarPointCloudEditorViewportClient(TWeakPtr<FLidarPointCloudEditor> InPointCloudEditor, const TSharedRef<SLidarPointCloudEditorViewport>& InPointCloudEditorViewport, FAdvancedPreviewScene* InPreviewScene, ULidarPointCloud* InPreviewPointCloud, ULidarPointCloudComponent* InPreviewPointCloudComponent)
 	: FEditorViewportClient(nullptr, InPreviewScene, StaticCastSharedRef<SEditorViewport>(InPointCloudEditorViewport))
 	, PointCloudComponent(InPreviewPointCloudComponent)
 	, PointCloudEditorPtr(InPointCloudEditor)
 	, PointCloudEditorViewportPtr(InPointCloudEditorViewport)
-	, SelectionMethod(ELidarPointCloudSelectionMethod::Box)
-	, SelectionMode(ELidarPointCloudSelectionMode::None)
-	, PaintingRadius(500)
 {
 	// Setup defaults for the common draw helper.
 	DrawHelper.bDrawPivot = false;
@@ -122,354 +70,23 @@ FLidarPointCloudEditorViewportClient::~FLidarPointCloudEditorViewportClient()
 	UAssetViewerSettings::Get()->OnAssetViewerSettingsChanged().RemoveAll(this);
 }
 
-void FLidarPointCloudEditorViewportClient::Tick(float DeltaSeconds)
-{
-	using UE::Lidar::Private::Editor::PaintMaxDeviation;
-
-	FEditorViewportClient::Tick(DeltaSeconds);
-
-	// Tick the preview scene world.
-	if (!GIntraFrameDebuggingGameThread)
-	{
-		PreviewScene->GetWorld()->Tick(LEVELTICK_All, DeltaSeconds);
-	}
-
-	// Process line traces if in Paint mode
-	if (SelectionMethod == ELidarPointCloudSelectionMethod::Paint)
-	{
-		if (TSharedPtr<FLidarPointCloudEditor> Editor = PointCloudEditorPtr.Pin())
-		{
-			ULidarPointCloud* PC = Editor->GetPointCloudBeingEdited();
-
-			if (Editor->IsEditMode())
-			{
-				const bool bPainting = Viewport->KeyState(EKeys::LeftMouseButton) && bLineTraceHit;
-				const float TraceRadius = FMath::Max(PC->GetEstimatedPointSpacing(), 0.5f);
-
-				const FLidarPointCloudRay Ray = DeprojectCurrentMousePosition();
-				if (FLidarPointCloudPoint* Point = PC->LineTraceSingle(Ray, TraceRadius, true))
-				{
-					const float NewDistance = FVector3f::Dist(Point->Location, Ray.Origin);
-					const float Deviation = (NewDistance - LineTraceDistance) / LineTraceDistance;
-
-					// If painting, prevent large depth changes
-					// If not, query larger trace radius - if it passes the deviation test, it was a gap
-					if (Deviation > PaintMaxDeviation && (bPainting || (FVector3f::Dist(PC->LineTraceSingle(Ray, TraceRadius * 6, true)->Location, Ray.Origin) - LineTraceDistance) / LineTraceDistance <= PaintMaxDeviation))
-					{
-						LineTraceHitPoint = FVector(Ray.Origin + Ray.GetDirection() * LineTraceDistance);
-					}
-					else
-					{
-						LineTraceHitPoint = (FVector)Point->Location;
-						LineTraceDistance = NewDistance;
-					}
-
-					bLineTraceHit = true;
-				}
-				else if(!bPainting)
-				{
-					bLineTraceHit = false;
-				}
-			}
-		}
-	}
-}
-
-bool FLidarPointCloudEditorViewportClient::InputKey(const FInputKeyEventArgs& EventArgs)
-{
-	using UE::Lidar::Private::Editor::PolySnapDistanceSq;
-
-	bool bHandled = false;
-
-	const bool bAlt = EventArgs.Key == EKeys::LeftAlt || EventArgs.Key == EKeys::RightAlt;
-	const bool bCtrl = EventArgs.Key == EKeys::LeftControl || EventArgs.Key == EKeys::RightControl;
-
-	if (TSharedPtr<FLidarPointCloudEditor> Editor = PointCloudEditorPtr.Pin())
-	{
-		// Edit Mode
-		if (Editor->IsEditMode())
-		{
-			if (EventArgs.Event == IE_Pressed)
-			{
-				if (EventArgs.Key == EKeys::Delete)
-				{
-					if (Viewport->KeyState(EKeys::LeftShift) || Viewport->KeyState(EKeys::RightShift))
-					{
-						Editor->DeletePoints();
-					}
-					else
-					{
-						Editor->HidePoints();
-					}
-
-					bHandled = true;
-				}
-				else if (EventArgs.Key == EKeys::Escape)
-				{
-					Editor->DeselectPoints();
-					SelectionPoints.Empty();
-					bHandled = true;
-				}
-				else if (EventArgs.Key == EKeys::Enter)
-				{
-					if (SelectionMethod == ELidarPointCloudSelectionMethod::Polygonal)
-					{
-						OnPolygonalSelectionEnd();
-						SelectionPoints.Empty();
-						bHandled = true;
-					}
-				}
-				else if (bAlt)
-				{
-					bHandled = true;
-					SelectionMode = ELidarPointCloudSelectionMode::Subtract;
-				}
-				else if (bCtrl)
-				{
-					bHandled = true;
-					SelectionMode = ELidarPointCloudSelectionMode::Add;
-				}
-				else if (EventArgs.Key == EKeys::LeftMouseButton)
-				{
-					bHandled = true;
-
-					// Selection start
-					if (SelectionMethod == ELidarPointCloudSelectionMethod::Box)
-					{
-						// Mark the cursor location for selection start
-						EventArgs.Viewport->GetMousePos(SelectionPoints[SelectionPoints.AddUninitialized()]);
-					}
-					else if (SelectionMethod == ELidarPointCloudSelectionMethod::Polygonal)
-					{
-						// Add new lasso point
-						FIntPoint NewPoint;
-						EventArgs.Viewport->GetMousePos(NewPoint);
-
-						// Don't allow duplicates
-						if (SelectionPoints.Num() == 0 || SelectionPoints.Last() != NewPoint)
-						{
-							TArray<FVector2D> VectorPoints = ToVectorArray(SelectionPoints);
-							VectorPoints.Add(NewPoint);
-
-							if (!FLidarPointCloudEditorHelper::IsPolygonSelfIntersecting(VectorPoints, false))
-							{
-								// Snap to first point
-								if (SelectionPoints.Num() > 1 && (NewPoint - SelectionPoints[0]).SizeSquared() < PolySnapDistanceSq)
-								{
-									OnPolygonalSelectionEnd();
-									SelectionPoints.Empty();
-								}
-								else
-								{
-									SelectionPoints.Add(NewPoint);
-								}
-							}
-						}
-					}
-					else if (SelectionMethod == ELidarPointCloudSelectionMethod::Lasso)
-					{
-						EventArgs.Viewport->GetMousePos(SelectionPoints[SelectionPoints.AddUninitialized()]);
-					}
-					else if (SelectionMethod == ELidarPointCloudSelectionMethod::Paint)
-					{
-						OnPaintSelection();
-					}
-				}
-
-				if (SelectionMethod == ELidarPointCloudSelectionMethod::Paint)
-				{
-					// Do not block ability to change camera speed
-					if (!Viewport->KeyState(EKeys::RightMouseButton))
-					{
-						if (EventArgs.Key == EKeys::MouseScrollUp)
-						{
-							PaintingRadius *= 1.1f;
-							bHandled = true;
-						}
-						else if (EventArgs.Key == EKeys::MouseScrollDown)
-						{
-							PaintingRadius /= 1.1f;
-							bHandled = true;
-						}
-					}
-				}
-			}
-			else if (EventArgs.Event == IE_Released)
-			{
-				if (bAlt || bCtrl)
-				{
-					SelectionMode = ELidarPointCloudSelectionMode::None;
-					bHandled = true;
-				}
-				else if (EventArgs.Key == EKeys::LeftMouseButton)
-				{
-					bHandled = true;
-
-					if (SelectionMethod == ELidarPointCloudSelectionMethod::Box)
-					{
-						// Mark the cursor location for selection end
-						EventArgs.Viewport->GetMousePos(SelectionPoints[SelectionPoints.AddUninitialized()]);
-						OnBoxSelectionEnd();
-						SelectionPoints.Empty();
-					}
-					else if (SelectionMethod == ELidarPointCloudSelectionMethod::Lasso)
-					{
-						OnLassoSelectionEnd();
-						SelectionPoints.Empty();
-					}
-				}
-			}
-		}
-		// Navigation Mode
-		else
-		{
-
-		}
-	}
-
-	if (!bHandled)
-	{
-		bHandled = FEditorViewportClient::InputKey(EventArgs);
-
-		// Handle viewport screenshot.
-		bHandled |= InputTakeScreenshot(EventArgs.Viewport, EventArgs.Key, EventArgs.Event);
-
-		bHandled |= AdvancedPreviewScene->HandleInputKey(EventArgs);
-
-	}
-
-	return bHandled;
-}
-
-bool FLidarPointCloudEditorViewportClient::InputAxis(FViewport* InViewport, FInputDeviceId DeviceId, FKey Key, float Delta, float DeltaTime, int32 NumSamples, bool bGamepad)
-{
-	using UE::Lidar::Private::Editor::LassoSpacingSq;
-
-	bool bHandled = false;
-
-	if (!bDisableInput)
-	{
-		if (TSharedPtr<FLidarPointCloudEditor> Editor = PointCloudEditorPtr.Pin())
-		{
-			if (Editor->IsEditMode())
-			{
-				if (Viewport->KeyState(EKeys::LeftMouseButton))
-				{
-					if (SelectionMethod == ELidarPointCloudSelectionMethod::Lasso)
-					{
-						FIntPoint NewPoint;
-						InViewport->GetMousePos(NewPoint);
-
-						// Check if the spacing is sufficient
-						if ((NewPoint - SelectionPoints.Last()).SizeSquared() > LassoSpacingSq)
-						{
-							SelectionPoints.Add(NewPoint);
-						}
-
-						bHandled = true;
-					}
-					else if (SelectionMethod == ELidarPointCloudSelectionMethod::Paint)
-					{
-						OnPaintSelection();
-						bHandled = true;
-					}
-				}
-			}
-		}
-
-		if(!bHandled)
-		{
-			bHandled = AdvancedPreviewScene->HandleViewportInput(InViewport, DeviceId, Key, Delta, DeltaTime, NumSamples, bGamepad);
-			if (bHandled)
-			{
-				Invalidate();
-			}
-			else
-			{
-				bHandled = FEditorViewportClient::InputAxis(InViewport, DeviceId, Key, Delta, DeltaTime, NumSamples, bGamepad);
-			}
-		}
-	}
-
-	return bHandled;
-}
-
 void FLidarPointCloudEditorViewportClient::DrawCanvas(FViewport& InViewport, FSceneView& View, FCanvas& Canvas)
 {
 	TArray<SLidarPointCloudEditorViewport::FOverlayTextItem> TextItems;
 
 	if (TSharedPtr<FLidarPointCloudEditor> Editor = PointCloudEditorPtr.Pin())
 	{
-		if (Editor->IsEditMode())
+		if(const ULidarPointCloud* PC = Editor->GetPointCloudBeingEdited())
 		{
-			TArray<FString> Labels;
-
-			// Add mode-specific labels
-			switch (SelectionMethod)
-			{
-			case ELidarPointCloudSelectionMethod::Box:
-				DrawSelectionBox(Canvas);
-				Labels.Append({
-					"BOX SELECTION MODE",
-					"",
-					"Click + Drag to replace selection",
-					"[CTRL] to add selection",
-					"[ALT] to subtract selection"
-				});
-				break;
-
-			case ELidarPointCloudSelectionMethod::Polygonal:
-				DrawSelectionPolygonal(Canvas);
-				Labels.Append({
-					"POLYGONAL SELECTION MODE",
-					"",
-					"[CTRL] to add selection",
-					"[ALT] to subtract selection"
-					});
-				break;
-
-			case ELidarPointCloudSelectionMethod::Lasso:
-				DrawSelectionLasso(Canvas);
-				Labels.Append({
-					"LASSO SELECTION MODE",
-					"",
-					"Click + Drag to paint selection shape",
-					"[CTRL] to add selection",
-					"[ALT] to subtract selection"
-				});
-				break;
+			const FIntVector BoundingSize(PC->GetBounds().GetSize());
 
-			case ELidarPointCloudSelectionMethod::Paint:
-				DrawSelectionPaint(Canvas);
-				Labels.Append({
-					"PAINT SELECTION MODE",
-					"",
-					"Click + Drag to paint selection",
-					"[SCROLL] to change brush size",
-					"[ALT] to subtract selection"
-				});
-				break;
-
-			default:
-				break;
-			}
-
-			// Add common labels
-			Labels.Append({
-				"",
-				"[ESCAPE] to de-select all points",
-				"[DELETE] to hide selected points",
-				"[SHIFT] + [DELETE] to permanently delete selected points",
-				"",
-				FString::Printf(TEXT("Selected Points: %d"), Editor->GetSelectedPoints().Num())
-			});
-
-			// Convert to text entries
-			TextItems.Reserve(Labels.Num());
-			for (const FString& Label : Labels)
-			{
-				TextItems.Add(SLidarPointCloudEditorViewport::FOverlayTextItem(FText::FromString(Label)));
-			}
+#define AS_NUM(Val) FText::AsNumber(Val, &FNumberFormattingOptions::DefaultWithGrouping())
+			TextItems.Emplace(FText::Format(LOCTEXT("PointCount", "Point Count: {0}"), AS_NUM(PC->GetNumPoints())));
+			TextItems.Emplace(FText::Format(LOCTEXT("NodeCount", "Node Count: {0}"), AS_NUM(PC->GetNumNodes())));
+			TextItems.Emplace(FText::Format(LOCTEXT("Bounds", "Bounds: {0} x {1} x {2}"), AS_NUM(BoundingSize.X), AS_NUM(BoundingSize.Y), AS_NUM(BoundingSize.Z)));
+			TextItems.Emplace(FText::Format(LOCTEXT("Collider", "Collider: {0} poly"), AS_NUM(PC->GetColliderPolys())));
+			TextItems.Emplace(FText::Format(LOCTEXT("DataSize", "Data Size: {0} MB"), AS_NUM(PC->GetDataSize())));
+#undef AS_NUM
 		}
 	}
 
@@ -481,39 +98,9 @@ void FLidarPointCloudEditorViewportClient::DrawCanvas(FViewport& InViewport, FSc
 
 bool FLidarPointCloudEditorViewportClient::ShouldOrbitCamera() const
 {
-	if (TSharedPtr<FLidarPointCloudEditor> Editor = PointCloudEditorPtr.Pin())
-	{
-		if (Editor->IsEditMode())
-		{
-			return false;
-		}
-	}
-
 	return GetDefault<ULevelEditorViewportSettings>()->bUseUE3OrbitControls || FEditorViewportClient::ShouldOrbitCamera();
 }
 
-void FLidarPointCloudEditorViewportClient::LostFocus(FViewport* InViewport)
-{
-	FEditorViewportClient::LostFocus(InViewport);
-
-	// Cancel selection
-	SelectionMode = ELidarPointCloudSelectionMode::None;
-	SelectionPoints.Empty();
-}
-
-void FLidarPointCloudEditorViewportClient::ReceivedFocus(FViewport* InViewport)
-{
-	// This is needed if the user presses Alt / Ctrl / Shift before the client acquires focus
-	if (Viewport->KeyState(EKeys::LeftAlt) || Viewport->KeyState(EKeys::RightAlt))
-	{
-		SelectionMode = ELidarPointCloudSelectionMode::Subtract;
-	}
-	else if (Viewport->KeyState(EKeys::LeftControl) || Viewport->KeyState(EKeys::RightControl))
-	{
-		SelectionMode = ELidarPointCloudSelectionMode::Add;
-	}
-}
-
 void FLidarPointCloudEditorViewportClient::PerspectiveCameraMoved()
 {
 	FEditorViewportClient::PerspectiveCameraMoved();
@@ -554,423 +141,11 @@ void FLidarPointCloudEditorViewportClient::SetAdvancedShowFlagsForScene(const bo
 	}
 }
 
-FSceneView* FLidarPointCloudEditorViewportClient::GetView()
-{
-	if (TSharedPtr<FLidarPointCloudEditor> Editor = PointCloudEditorPtr.Pin())
-	{
-		// Compute a view.
-		FSceneViewFamilyContext ViewFamily(FSceneViewFamily::ConstructionValues(Viewport, GetScene(), EngineShowFlags).SetRealtimeUpdate(IsRealtime()));
-		FSceneView* View = CalcSceneView(&ViewFamily);
-
-		const FVector LocationOffset = Editor->GetPointCloudBeingEdited()->LocationOffset;
-
-		// Adjust for the LocationOffset
-		if (View->IsPerspectiveProjection())
-		{
-			View->ViewLocation -= LocationOffset;
-			View->UpdateViewMatrix();
-		}
-
-		return View;
-	}
-
-	return nullptr;
-}
-
-FLidarPointCloudRay FLidarPointCloudEditorViewportClient::DeprojectCurrentMousePosition()
-{
-	FSceneView* View = GetView();
-	const FMatrix InvViewProjectionMatrix = View->ViewMatrices.GetInvViewProjectionMatrix();
-
-	FIntPoint CurrentMousePosition;
-	Viewport->GetMousePos(CurrentMousePosition);
-
-	FVector Origin, Direction;
-	FSceneView::DeprojectScreenToWorld(FVector2D(CurrentMousePosition), FIntRect(FIntPoint(0, 0), Viewport->GetSizeXY()), InvViewProjectionMatrix, Origin, Direction);
-
-	return FLidarPointCloudRay((FVector3f)Origin, (FVector3f)Direction);
-}
-
-void FLidarPointCloudEditorViewportClient::OnBoxSelectionEnd()
-{
-	if (ULidarPointCloudComponent* PointCloudComponentRawPtr = PointCloudComponent.Get())
-	{
-		if (TSharedPtr<FLidarPointCloudEditor> Editor = PointCloudEditorPtr.Pin())
-		{
-			if (SelectionPoints[0] == SelectionPoints.Last())
-			{
-				Editor->DeselectPoints();
-				return;
-			}
-
-			FIntVector4 SelectionArea;
-			SelectionArea.X = FMath::Min(SelectionPoints[0].X, SelectionPoints[1].X);
-			SelectionArea.Y = FMath::Min(SelectionPoints[0].Y, SelectionPoints[1].Y);
-			SelectionArea.Z = FMath::Max(SelectionPoints[0].X, SelectionPoints[1].X);
-			SelectionArea.W = FMath::Max(SelectionPoints[0].Y, SelectionPoints[1].Y);
-
-			const FConvexVolume ConvexVolume = BuildConvexVolumeForPoints(TArray<FVector2D>({
-				FVector2D(SelectionArea.X, SelectionArea.Y),
-				FVector2D(SelectionArea.X, SelectionArea.W),
-				FVector2D(SelectionArea.Z, SelectionArea.W),
-				FVector2D(SelectionArea.Z, SelectionArea.Y) }));
-
-			if (SelectionMode == ELidarPointCloudSelectionMode::Subtract)
-			{
-				Editor->DeselectPointsByConvexVolume(ConvexVolume);
-			}
-			else
-			{
-				Editor->SelectPointsByConvexVolume(ConvexVolume, SelectionMode == ELidarPointCloudSelectionMode::Add);
-			}
-		}
-	}
-}
-
-void FLidarPointCloudEditorViewportClient::OnPolygonalSelectionEnd()
-{
-	// Skip invalid selections
-	if (SelectionPoints.Num() < 3)
-	{
-		return;
-	}
-
-	if (ULidarPointCloudComponent* PointCloudComponentRawPtr = PointCloudComponent.Get())
-	{
-		if (TSharedPtr<FLidarPointCloudEditor> Editor = PointCloudEditorPtr.Pin())
-		{
-			TArray<TArray<FVector2D>> ConvexShapes;
-			TArray<FVector2D> VectorPoints = ToVectorArray(SelectionPoints);
-
-			if (IsPolygonConvex(SelectionPoints))
-			{
-				ConvexShapes.Add(VectorPoints);
-			}
-			else
-			{
-				// Check for self-intersecting shape
-				if (!FLidarPointCloudEditorHelper::IsPolygonSelfIntersecting(VectorPoints, true))
-				{
-					// The separation needs points in CCW order
-					if (!FGeomTools2D::IsPolygonWindingCCW(VectorPoints))
-					{
-						Algo::Reverse(VectorPoints);
-					}
-
-					TArray<FVector2D> Triangles;
-					FGeomTools2D::TriangulatePoly(Triangles, VectorPoints, false);
-					FGeomTools2D::GenerateConvexPolygonsFromTriangles(ConvexShapes, Triangles);
-				}
-			}
-
-			for (int32 i = 0; i < ConvexShapes.Num(); ++i)
-			{
-				const FConvexVolume ConvexVolume = BuildConvexVolumeForPoints(ConvexShapes[i]);
-
-				if (SelectionMode == ELidarPointCloudSelectionMode::Subtract)
-				{
-					Editor->DeselectPointsByConvexVolume(ConvexVolume);
-				}
-				else
-				{
-					// Consecutive shapes need to be additive
-					Editor->SelectPointsByConvexVolume(ConvexVolume, i > 0 || SelectionMode == ELidarPointCloudSelectionMode::Add);
-				}
-			}
-		}
-	}
-}
-
-void FLidarPointCloudEditorViewportClient::OnLassoSelectionEnd()
-{
-	OnPolygonalSelectionEnd();
-}
-
-void FLidarPointCloudEditorViewportClient::OnPaintSelection()
-{
-	if (TSharedPtr<FLidarPointCloudEditor> Editor = PointCloudEditorPtr.Pin())
-	{
-		if (bLineTraceHit)
-		{
-			if (SelectionMode == ELidarPointCloudSelectionMode::Subtract)
-			{
-				Editor->DeselectPointsBySphere(FSphere(LineTraceHitPoint, PaintingRadius));
-			}
-			else
-			{
-				Editor->SelectPointsBySphere(FSphere(LineTraceHitPoint, PaintingRadius));
-			}
-		}
-	}
-}
-
-void FLidarPointCloudEditorViewportClient::DrawSelectionBox(FCanvas& Canvas)
-{
-	if (SelectionPoints.Num() == 0)
-	{
-		return;
-	}
-
-	const FIntPoint SelectionStartLocation = SelectionPoints[0];
-	FIntPoint SelectionCurrentLocation;
-	Viewport->GetMousePos(SelectionCurrentLocation);
-
-	const float InvScale = 1.0f / Viewport->GetClient()->GetDPIScale();
-	
-	const float X = FMath::Min(SelectionStartLocation.X, SelectionCurrentLocation.X) * InvScale;
-	const float Y = FMath::Min(SelectionStartLocation.Y, SelectionCurrentLocation.Y) * InvScale;
-	const float SizeX = FMath::Max(SelectionStartLocation.X, SelectionCurrentLocation.X) * InvScale - X;
-	const float SizeY = FMath::Max(SelectionStartLocation.Y, SelectionCurrentLocation.Y) * InvScale - Y;
-
-	FLinearColor SelectionColor = GetDefault<UEditorStyleSettings>()->SelectionColor;
-	SelectionColor.A = 0.35f;
-
-	Canvas.DrawTile(X, Y, SizeX, SizeY, 0, 0, 0, 0, SelectionColor);
-
-	// Selection Border
-	{
-		FCanvasLineItem Line;
-		Line.SetColor(GetDefault<UEditorStyleSettings>()->SelectionColor);
-		Line.LineThickness = 2;
-
-		Line.Origin = FVector(X, Y, 0);
-		Line.EndPos = FVector(X + SizeX, Y, 0);
-		Canvas.DrawItem(Line);
-
-		Line.Origin = Line.EndPos;
-		Line.EndPos = FVector(X + SizeX, Y + SizeY, 0);
-		Canvas.DrawItem(Line);
-
-		Line.Origin = Line.EndPos;
-		Line.EndPos = FVector(X, Y + SizeY, 0);
-		Canvas.DrawItem(Line);
-
-		Line.Origin = Line.EndPos;
-		Line.EndPos = FVector(X, Y, 0);
-		Canvas.DrawItem(Line);
-	}
-}
-
-void FLidarPointCloudEditorViewportClient::DrawSelectionPolygonal(FCanvas& Canvas)
-{
-	using UE::Lidar::Private::Editor::PolySnapDistanceSq;
-
-	if (SelectionPoints.Num() == 0)
-	{
-		return;
-	}
-
-	// Create a local copy of selection points, injecting the current mouse position at the end
-	TArray<FIntPoint> DrawSelectionPoints = SelectionPoints;
-	Viewport->GetMousePos(DrawSelectionPoints[DrawSelectionPoints.AddUninitialized()]);
-
-	if (DrawSelectionPoints.Last() == DrawSelectionPoints.Last(1))
-	{
-		DrawSelectionPoints.RemoveAt(DrawSelectionPoints.Num() - 1, 1, false);
-	}
-
-	TArray<FVector2D> VectorPoints = ToVectorArray(DrawSelectionPoints);
-
-	// Account for DPI
-	const float InvScale = 1.0f / Viewport->GetClient()->GetDPIScale();
-	for(FVector2D& DrawPoint : VectorPoints)
-	{
-		DrawPoint *= InvScale;
-	}
-
-	// Calculate visual indication of complete polygon for the user
-	const bool bPolyComplete = DrawSelectionPoints.Num() > 2 && (DrawSelectionPoints.Last() - DrawSelectionPoints[0]).SizeSquared() < PolySnapDistanceSq;
-	const bool bSelfIntersecting = DrawSelectionPoints.Num() > 2 && FLidarPointCloudEditorHelper::IsPolygonSelfIntersecting(VectorPoints, true);
-	FLinearColor SelectionColor = bSelfIntersecting ? FLinearColor::Red : bPolyComplete ? FLinearColor::Green : GetDefault<UEditorStyleSettings>()->SelectionColor;
-
-	// Selection Area
-	if (VectorPoints.Num() > 2 && !bSelfIntersecting)
-	{
-		FCanvasUVTri Tri;
-		Tri.V0_Pos = Tri.V1_Pos = Tri.V2_Pos = Tri.V0_UV = Tri.V1_UV = Tri.V2_UV = FVector2D::ZeroVector;
-		SelectionColor.A = 0.35f;
-		Tri.V0_Color = Tri.V1_Color = Tri.V2_Color = SelectionColor;
-
-		TArray<FCanvasUVTri> TriangleList;
-		TArray<TArray<FVector2D>> Polygons;
-
-		if (IsPolygonConvex(DrawSelectionPoints))
-		{
-			Polygons.Add(VectorPoints);
-			TriangleList.Reserve(VectorPoints.Num() - 2);
-		}
-		else
-		{
-			// The separation needs points in CCW order
-			if (!FGeomTools2D::IsPolygonWindingCCW(VectorPoints))
-			{
-				Algo::Reverse(VectorPoints);
-			}
-
-			TArray<FVector2D> Triangles;
-			if (FGeomTools2D::TriangulatePoly(Triangles, VectorPoints, false))
-			{
-				FGeomTools2D::GenerateConvexPolygonsFromTriangles(Polygons, Triangles);
-
-				// Calculate the number of triangles and reserve space
-				int32 NumTriangles = 0;
-				for (const TArray<FVector2D>& Polygon : Polygons)
-				{
-					NumTriangles += Polygon.Num() - 2;
-				}
-				TriangleList.Reserve(NumTriangles);
-			}
-		}
-
-		for (const TArray<FVector2D>& Polygon : Polygons)
-		{
-			for (int32 i = 2; i < Polygon.Num(); ++i)
-			{
-				Tri.V0_Pos = Polygon[i];
-				Tri.V1_Pos = Polygon[0];
-				Tri.V2_Pos = Polygon[i - 1];
-				TriangleList.Add(Tri);
-			}
-		}
-
-		if (TriangleList.Num() > 0)
-		{
-			FCanvasTriangleItem Selection(TriangleList, GWhiteTexture);
-			Selection.BlendMode = SE_BLEND_AlphaBlend;
-			Canvas.DrawItem(Selection);
-		}
-	}
-
-	// Selection Border
-	{
-		FCanvasLineItem Line;
-		Line.SetColor(SelectionColor);
-		Line.LineThickness = 2;
-
-		for (int32 i = 1; i < VectorPoints.Num(); ++i)
-		{
-			Line.Origin = FVector(VectorPoints[i], 0);
-			Line.EndPos = FVector(VectorPoints[i - 1], 0);
-			Canvas.DrawItem(Line);
-		}
-	}
-}
-
-void FLidarPointCloudEditorViewportClient::DrawSelectionLasso(FCanvas& Canvas)
-{
-	DrawSelectionPolygonal(Canvas);
-}
-
-void FLidarPointCloudEditorViewportClient::DrawSelectionPaint(FCanvas& Canvas)
-{
-	if (TSharedPtr<FLidarPointCloudEditor> Editor = PointCloudEditorPtr.Pin())
-	{
-		if (TSharedPtr<SLidarPointCloudEditorViewport> EditorViewport = PointCloudEditorViewportPtr.Pin())
-		{
-			EditorViewport->PaintBrush->SetVisibility(bLineTraceHit);
-			EditorViewport->PaintBrush->SetWorldScale3D(FVector(PaintingRadius * 0.02f));
-			EditorViewport->PaintBrush->SetWorldLocation(LineTraceHitPoint);
-		}
-	}
-}
-
-FConvexVolume FLidarPointCloudEditorViewportClient::BuildConvexVolumeForPoints(const TArray<FVector2D>& Points)
-{
-	FConvexVolume ConvexVolume;
-
-	if (TSharedPtr<FLidarPointCloudEditor> Editor = PointCloudEditorPtr.Pin())
-	{
-		FSceneView* View = GetView();
-		const FVector LocationOffset = Editor->GetPointCloudBeingEdited()->LocationOffset;
-		const FMatrix InvViewProjectionMatrix = View->ViewMatrices.GetInvViewProjectionMatrix();
-
-		TArray<FVector> Origins; Origins.AddUninitialized(Points.Num() + 2);
-		TArray<FVector> Normals; Normals.AddUninitialized(Points.Num() + 2);
-		TArray<FVector> Directions; Directions.AddUninitialized(Points.Num());
-		FVector MeanCenter = FVector::ZeroVector;
-
-		for (int32 i = 0; i < Points.Num(); ++i)
-		{
-			FSceneView::DeprojectScreenToWorld(Points[i], FIntRect(FIntPoint(0, 0), Viewport->GetSizeXY()), InvViewProjectionMatrix, Origins[i], Directions[i]);
-			MeanCenter += Origins[i];
-		}
-
-		MeanCenter /= Points.Num();
-
-		const FVector& ViewDirection = View->GetViewDirection();
-
-		// Shared calculations
-		Normals.Last(1) = ViewDirection;
-		Normals.Last() = -ViewDirection;
-		Origins.Last(1) = Origins[0] + ViewDirection * 99999999.0f;
-
-		// Calculate plane normals
-		bool bFlipNormals = false;
-		for (int32 i = 0; i < Points.Num(); ++i)
-		{
-			Normals[i] = ((Origins[(i + 1) % Points.Num()] - Origins[i]).GetSafeNormal() ^ Directions[i]).GetSafeNormal();
-
-			if (i == 0)
-			{
-				bFlipNormals = FVector::DotProduct(Normals[i], (MeanCenter - Origins[i])) > 0;
-			}
-
-			if (bFlipNormals)
-			{
-				Normals[i] = -Normals[i];
-			}
-		}
-
-		// Perspective View
-		if (View->IsPerspectiveProjection())
-		{
-			Origins.Last() = Origins[0];
-		}
-		// Ortho Views
-		else
-		{
-			// Adjust for the LocationOffset
-			for (int32 i = 0; i < Points.Num(); ++i)
-			{
-				Origins[i] -= LocationOffset;
-			}
-
-			Origins.Last() = -Origins.Last(1);
-		}
-
-		for (int32 i = 0; i < Origins.Num(); ++i)
-		{
-			ConvexVolume.Planes.Emplace(Origins[i], Normals[i]);
-		}
-
-		ConvexVolume.Init();
-	}
-
-	return ConvexVolume;
-}
-
 void FLidarPointCloudEditorViewportClient::ResetCamera()
 {
 	if (ULidarPointCloudComponent* PointCloudComponentRawPtr = PointCloudComponent.Get())
 	{
-		FBox FocusBounds(EForceInit::ForceInit);
-
-		// Focus on selection, if possible
-		if (TSharedPtr<FLidarPointCloudEditor> Editor = PointCloudEditorPtr.Pin())
-		{
-			for (FLidarPointCloudPoint** Point = Editor->GetSelectedPoints().GetData(), **DataEnd = Point + Editor->GetSelectedPoints().Num(); Point != DataEnd; ++Point)
-			{
-				FocusBounds += (FVector)(*Point)->Location;
-			}
-		}
-
-		// Fallback to the whole cloud's bounds
-		if(!FocusBounds.IsValid)
-		{
-			FocusBounds = PointCloudComponentRawPtr->Bounds.GetBox();
-		}
-
-		FocusViewportOnBox(FocusBounds);
+		FocusViewportOnBox(PointCloudComponentRawPtr->Bounds.GetBox());
 		Invalidate();
 	}
 }
@@ -994,14 +169,4 @@ bool FLidarPointCloudEditorViewportClient::IsSetShowNodesChecked() const
 	return false;
 }
 
-void FLidarPointCloudEditorViewportClient::SetSelectionMethod(ELidarPointCloudSelectionMethod NewSelectionMethod)
-{
-	SelectionMethod = NewSelectionMethod;
-
-	if (TSharedPtr<SLidarPointCloudEditorViewport> EditorViewport = PointCloudEditorViewportPtr.Pin())
-	{
-		EditorViewport->PaintBrush->SetVisibility(SelectionMethod == ELidarPointCloudSelectionMethod::Paint);
-	}
-}
-
 #undef LOCTEXT_NAMESPACE 
diff --git a/Engine/Plugins/Enterprise/LidarPointCloud/Source/LidarPointCloudEditor/Private/LidarPointCloudEditorViewportClient.h b/Engine/Plugins/Enterprise/LidarPointCloud/Source/LidarPointCloudEditor/Private/LidarPointCloudEditorViewportClient.h
--- a/Engine/Plugins/Enterprise/LidarPointCloud/Source/LidarPointCloudEditor/Private/LidarPointCloudEditorViewportClient.h
+++ b/Engine/Plugins/Enterprise/LidarPointCloud/Source/LidarPointCloudEditor/Private/LidarPointCloudEditorViewportClient.h
@@ -16,14 +16,6 @@ class SLidarPointCloudEditorViewport;
 class ULidarPointCloud;
 class ULidarPointCloudComponent;
 
-enum class ELidarPointCloudSelectionMethod : uint8
-{
-	Box,
-	Polygonal,
-	Lasso,
-	Paint
-};
-
 /** Viewport Client for the preview viewport */
 class FLidarPointCloudEditorViewportClient : public FEditorViewportClient, public TSharedFromThis<FLidarPointCloudEditorViewportClient>
 {
@@ -32,9 +24,6 @@ public:
 	~FLidarPointCloudEditorViewportClient();
 
 	// FEditorViewportClient interface
-	virtual bool InputKey(const FInputKeyEventArgs& EventArgs) override;
-	virtual bool InputAxis(FViewport* InViewport, FInputDeviceId DeviceId, FKey Key, float Delta, float DeltaTime, int32 NumSamples = 1, bool bGamepad = false) override;
-	virtual void Tick(float DeltaSeconds) override;
 	virtual void DrawCanvas(FViewport& InViewport, FSceneView& View, FCanvas& Canvas) override;
 	virtual bool InputWidgetDelta(FViewport* InViewport, EAxisList::Type CurrentAxis, FVector& Drag, FRotator& Rot, FVector& Scale) override { return false; }
 	virtual void TrackingStarted(const struct FInputEventState& InInputState, bool bIsDragging, bool bNudge) override {}
@@ -47,9 +36,7 @@ public:
 	virtual FMatrix GetWidgetCoordSystem() const override { return FMatrix::Identity; }
 	virtual ECoordSystem GetWidgetCoordSystemSpace() const override { return COORD_Local; }
 	virtual bool ShouldOrbitCamera() const override;
-	virtual void LostFocus(FViewport* InViewport) override;
-	virtual void ReceivedFocus(FViewport* InViewport) override;
-
+	
 	void ResetCamera();
 
 	/** Callback for toggling the nodes show flag. */
@@ -58,9 +45,6 @@ public:
 	/** Callback for checking the nodes show flag. */
 	bool IsSetShowNodesChecked() const;
 
-	FORCEINLINE ELidarPointCloudSelectionMethod GetSelectionMethod() const { return SelectionMethod; }
-	void SetSelectionMethod(ELidarPointCloudSelectionMethod NewSelectionMethod);
-
 protected:
 	// FEditorViewportClient interface
 	virtual void PerspectiveCameraMoved() override;
@@ -70,22 +54,6 @@ protected:
 	/** Used to (re)-set the viewport show flags related to post processing*/
 	void SetAdvancedShowFlagsForScene(const bool bAdvancedShowFlags);
 
-	FSceneView* GetView();
-	FLidarPointCloudRay DeprojectCurrentMousePosition();
-
-private:
-	void OnBoxSelectionEnd();
-	void OnPolygonalSelectionEnd();
-	void OnLassoSelectionEnd();
-	void OnPaintSelection();
-
-	void DrawSelectionBox(FCanvas& Canvas);
-	void DrawSelectionPolygonal(FCanvas& Canvas);
-	void DrawSelectionLasso(FCanvas& Canvas);
-	void DrawSelectionPaint(FCanvas& Canvas);
-
-	FConvexVolume BuildConvexVolumeForPoints(const TArray<FVector2D>& Points);
-
 private:
 	/** Component for the point cloud. */
 	TWeakObjectPtr<ULidarPointCloudComponent> PointCloudComponent;
@@ -98,13 +66,4 @@ private:
 
 	/** Stored pointer to the preview scene in which the point cloud is shown */
 	FAdvancedPreviewScene* AdvancedPreviewScene;
-
-	ELidarPointCloudSelectionMethod SelectionMethod;
-	ELidarPointCloudSelectionMode SelectionMode;
-	TArray<FIntPoint> SelectionPoints;
-
-	float PaintingRadius;
-	bool bLineTraceHit;
-	FVector LineTraceHitPoint;
-	float LineTraceDistance;
 };
diff --git a/Engine/Plugins/Enterprise/LidarPointCloud/Source/LidarPointCloudRuntime/Private/LidarPointCloud.cpp b/Engine/Plugins/Enterprise/LidarPointCloud/Source/LidarPointCloudRuntime/Private/LidarPointCloud.cpp
--- a/Engine/Plugins/Enterprise/LidarPointCloud/Source/LidarPointCloudRuntime/Private/LidarPointCloud.cpp
+++ b/Engine/Plugins/Enterprise/LidarPointCloud/Source/LidarPointCloudRuntime/Private/LidarPointCloud.cpp
@@ -449,6 +449,11 @@ bool ULidarPointCloud::HasCollisionData() const
 	return Octree.HasCollisionData();
 }
 
+int32 ULidarPointCloud::GetColliderPolys() const
+{
+	return Octree.HasCollisionData() ? Octree.GetCollisionData()->Indices.Num() : 0;
+}
+
 void ULidarPointCloud::RefreshRendering()
 {
 	Octree.MarkRenderDataDirty();
diff --git a/Engine/Plugins/Enterprise/LidarPointCloud/Source/LidarPointCloudRuntime/Public/LidarPointCloud.h b/Engine/Plugins/Enterprise/LidarPointCloud/Source/LidarPointCloudRuntime/Public/LidarPointCloud.h
--- a/Engine/Plugins/Enterprise/LidarPointCloud/Source/LidarPointCloudRuntime/Public/LidarPointCloud.h
+++ b/Engine/Plugins/Enterprise/LidarPointCloud/Source/LidarPointCloudRuntime/Public/LidarPointCloud.h
@@ -221,6 +221,10 @@ public:
 	/** Returns true, if the Octree has collision built */
 	UFUNCTION(BlueprintPure, Category = "Lidar Point Cloud")
 	bool HasCollisionData() const;
+	
+	/** Returns the number of polygons in the collider or 0 if no collider is built */
+	UFUNCTION(BlueprintPure, Category = "Lidar Point Cloud")
+	int32 GetColliderPolys() const;
 
 	UFUNCTION(BlueprintCallable, Category = "Lidar Point Cloud")
 	void RefreshRendering();
