New Lidar EdMode UI

diff --git a/Engine/Plugins/Enterprise/LidarPointCloud/Source/LidarPointCloudEditor/Private/LidarPointCloudEdMode.cpp b/Engine/Plugins/Enterprise/LidarPointCloud/Source/LidarPointCloudEditor/Private/LidarPointCloudEdMode.cpp
--- a/Engine/Plugins/Enterprise/LidarPointCloud/Source/LidarPointCloudEditor/Private/LidarPointCloudEdMode.cpp
+++ b/Engine/Plugins/Enterprise/LidarPointCloud/Source/LidarPointCloudEditor/Private/LidarPointCloudEdMode.cpp
@@ -40,29 +40,14 @@ void ULidarEditorMode::Enter()
 
 #define REGISTER_TOOL(Tool) RegisterTool(Commands.Toolkit##Tool, TEXT("Lidar"#Tool"Tool"), NewObject<ULidarEditorToolBuilder##Tool>())
 	REGISTER_TOOL(Select);
-	REGISTER_TOOL(Align);
-	REGISTER_TOOL(Merge);
-	REGISTER_TOOL(Collision);
-	REGISTER_TOOL(Normals);
-	REGISTER_TOOL(Meshing);
 	REGISTER_TOOL(BoxSelection);
 	REGISTER_TOOL(PolygonalSelection);
 	REGISTER_TOOL(LassoSelection);
 	REGISTER_TOOL(PaintSelection);
 #undef REGISTER_TOOL
-	
-	PaletteChangedHandle = Toolkit->OnPaletteChanged().AddUObject(this, &ULidarEditorMode::UpdateOnPaletteChange);
-	
-	GetToolManager()->ConfigureChangeTrackingMode(EToolChangeTrackingMode::NoChangeTracking);
-	Toolkit->SetCurrentPalette(LidarEditorPalletes::Manage);
-	GetToolManager()->ConfigureChangeTrackingMode(EToolChangeTrackingMode::UndoToExit);
-}
 
-void ULidarEditorMode::Exit()
-{
-	Toolkit->OnPaletteChanged().Remove(PaletteChangedHandle);
-	
-	Super::Exit();
+	GetToolManager()->SelectActiveToolType(EToolSide::Mouse, "LidarSelectTool");
+	GetToolManager()->ActivateTool(EToolSide::Mouse);
 }
 
 bool ULidarEditorMode::IsSelectionAllowed(AActor* InActor, bool bInSelection) const
@@ -87,13 +72,7 @@ bool ULidarEditorMode::GetPivotForOrbit(FVector& OutPivot) const
 
 TMap<FName, TArray<TSharedPtr<FUICommandInfo>>> ULidarEditorMode::GetModeCommands() const
 {
-	return FLidarPointCloudEditorCommands::GetCommands();
-}
-
-void ULidarEditorMode::UpdateOnPaletteChange(FName NewPalette)
-{
-	GetToolManager()->SelectActiveToolType(EToolSide::Mouse, "LidarSelectTool");
-	GetToolManager()->ActivateTool(EToolSide::Mouse);
+	return FLidarPointCloudEditorCommands::GetToolkitCommands();
 }
 
 void ULidarEditorMode::CancelActiveToolAction()
@@ -123,9 +102,15 @@ void ULidarEditorMode::BindCommands()
 
 void ULidarEditorMode::OnToolStarted(UInteractiveToolManager* Manager, UInteractiveTool* Tool)
 {
-	if(ULidarEditorToolBase* LidarTool = Cast<ULidarEditorToolBase>(Tool))
+	if(const ULidarEditorToolBase* LidarTool = Cast<ULidarEditorToolBase>(Tool))
 	{
 		GetToolManager()->DisplayMessage(LidarTool->GetToolMessage(), EToolMessageLevel::UserNotification);
+
+		if(FLidarPointCloudEdModeToolkit* LidarToolkit = (FLidarPointCloudEdModeToolkit*)Toolkit.Get())
+		{
+			LidarToolkit->SetActorSelection(LidarTool->IsActorSelection());
+			LidarToolkit->SetBrushTool(Cast<ULidarEditorToolPaintSelection>(Tool));
+		}
 	}
 }
 
diff --git a/Engine/Plugins/Enterprise/LidarPointCloud/Source/LidarPointCloudEditor/Private/LidarPointCloudEdMode.h b/Engine/Plugins/Enterprise/LidarPointCloud/Source/LidarPointCloudEditor/Private/LidarPointCloudEdMode.h
--- a/Engine/Plugins/Enterprise/LidarPointCloud/Source/LidarPointCloudEditor/Private/LidarPointCloudEdMode.h
+++ b/Engine/Plugins/Enterprise/LidarPointCloud/Source/LidarPointCloudEditor/Private/LidarPointCloudEdMode.h
@@ -16,15 +16,12 @@ namespace FLidarEditorModes
 UCLASS()
 class ULidarEditorMode : public UBaseLegacyWidgetEdMode
 {
-	FDelegateHandle PaletteChangedHandle;
-
 public:
 	GENERATED_BODY()
 	
 	ULidarEditorMode();
 
 	virtual void Enter() override;
-	virtual void Exit() override;
 	virtual bool IsSelectionAllowed(AActor* InActor, bool bInSelection) const override;
 	virtual void CreateToolkit() override;
 	virtual bool GetPivotForOrbit(FVector& OutPivot) const override;
@@ -34,6 +31,5 @@ public:
 	virtual void OnToolEnded(UInteractiveToolManager* Manager, UInteractiveTool* Tool) override;
 
 private:
-	void UpdateOnPaletteChange(FName NewPalette);
 	void CancelActiveToolAction();
 };
\ No newline at end of file
diff --git a/Engine/Plugins/Enterprise/LidarPointCloud/Source/LidarPointCloudEditor/Private/LidarPointCloudEdModeToolkit.cpp b/Engine/Plugins/Enterprise/LidarPointCloud/Source/LidarPointCloudEditor/Private/LidarPointCloudEdModeToolkit.cpp
--- a/Engine/Plugins/Enterprise/LidarPointCloud/Source/LidarPointCloudEditor/Private/LidarPointCloudEdModeToolkit.cpp
+++ b/Engine/Plugins/Enterprise/LidarPointCloud/Source/LidarPointCloudEditor/Private/LidarPointCloudEdModeToolkit.cpp
@@ -1,18 +1,389 @@
 // Copyright Epic Games, Inc. All Rights Reserved.
 
 #include "LidarPointCloudEdModeToolkit.h"
-#include "EditorModeManager.h"
 #include "EdModeInteractiveToolsContext.h"
 #include "InteractiveToolManager.h"
 #include "LidarPointCloudEditorCommands.h"
+#include "LidarPointCloudEditorTools.h"
 #include "LidarPointCloudEdMode.h"
-#include "LidarPointCloudShared.h"
 #include "Selection.h"
+#include "SlateOptMacros.h"
 #include "StatusBarSubsystem.h"
 #include "Toolkits/AssetEditorModeUILayer.h"
 
+#include "Widgets/Input/SNumericEntryBox.h"
+#include "Widgets/Input/SButton.h"
+#include "Widgets/Layout/SWrapBox.h"
+#include "Widgets/Layout/SSpacer.h"
+
 #define LOCTEXT_NAMESPACE "LidarEditMode"
 
+BEGIN_SLATE_FUNCTION_BUILD_OPTIMIZATION
+void SLidarEditorWidget::Construct(const FArguments& InArgs)
+{
+	LidarEditorMode = (ULidarEditorMode*)GLevelEditorModeTools().GetActiveMode(FLidarEditorModes::EM_Lidar);
+	
+	// Everything (or almost) uses this padding, change it to expand the padding.
+	FMargin StandardPadding(5.f, 2.f);
+	FMargin HeaderPadding(8.f, 4.f);
+
+	FSlateFontInfo HeaderFont = FAppStyle::GetFontStyle(TEXT("DetailsView.CategoryFontStyle")); HeaderFont.Size = 11;
+	FSlateFontInfo StandardFont = FAppStyle::GetFontStyle(TEXT("PropertyWindow.NormalFont"));
+	FSlateFontInfo LabelFont = FAppStyle::GetFontStyle(TEXT("PropertyWindow.NormalFont")); LabelFont.Size = 9;
+
+#define HEADEREX(ContText, Category) +SVerticalBox::Slot().AutoHeight().Padding(FMargin(0.0f, 5.0f))[SNew(SHeaderRow)+SHeaderRow::Column(ContText).HAlignCell(HAlign_Left).FillWidth(1).HeaderContentPadding(HeaderPadding)[SNew(STextBlock).Text(LOCTEXT(Category, ContText)).Font(HeaderFont)]]
+#define HEADER(ContText) HEADEREX(ContText, ContText"Header")
+#define BUTTON(Category, ConText, Action) +SHorizontalBox::Slot().Padding(StandardPadding).FillWidth(0.5f)[SNew(SButton).HAlign(HAlign_Center).Text(LOCTEXT(Category,ConText)).OnClicked_Lambda([this]{ Action return FReply::Handled(); })]
+#define BUTTON_POINTS(Category, ConText, Action) +SHorizontalBox::Slot().Padding(StandardPadding).FillWidth(0.5f)[SNew(SButton).HAlign(HAlign_Center).Text(LOCTEXT(Category,ConText)).OnClicked_Lambda([this]{ Action return FReply::Handled(); }).IsEnabled(this, &SLidarEditorWidget::IsPointSelection)]
+#define BUTTON_ACTORS(Category, ConText, Action) +SHorizontalBox::Slot().Padding(StandardPadding).FillWidth(0.5f)[SNew(SButton).HAlign(HAlign_Center).Text(LOCTEXT(Category,ConText)).OnClicked_Lambda([this]{ Action return FReply::Handled(); }).IsEnabled(this, &SLidarEditorWidget::IsActorSelection)]
+
+	ChildSlot
+	[
+		SNew(SVerticalBox)
+		+ SVerticalBox::Slot()
+		.AutoHeight()
+		[
+			SNew(SHorizontalBox)
+			BUTTON("SelectionClear", "Clear",
+			{
+				if(IsActorSelection())
+				{
+					FLidarPointCloudEditorHelper::ClearActorSelection();
+				}
+				else
+				{
+					FLidarPointCloudEditorHelper::ClearSelection();
+				}
+			})
+			BUTTON("SelectionInvert", "Invert",
+			{
+				if(IsActorSelection())
+				{
+					FLidarPointCloudEditorHelper::InvertActorSelection();
+				}
+				else
+				{
+					FLidarPointCloudEditorHelper::InvertSelection();
+				}
+			})
+		]
+		+SVerticalBox::Slot()
+		.AutoHeight()
+		.Padding(FMargin(0.0f, 5.0f))
+		[
+			SNew(SHeaderRow)
+			.Visibility(this, &SLidarEditorWidget::GetBrushVisibility)
+			+SHeaderRow::Column("Brush")
+			.HAlignCell(HAlign_Left)
+			.FillWidth(1)
+			.HeaderContentPadding(HeaderPadding)
+			[
+				SNew(STextBlock)
+				.Text(LOCTEXT("BrushHeader", "Brush"))
+				.Font(HeaderFont)
+			]
+		]
+		+ SVerticalBox::Slot()
+		.AutoHeight()
+		[
+			SNew(SHorizontalBox)
+			.Visibility(this, &SLidarEditorWidget::GetBrushVisibility)
+			+ SHorizontalBox::Slot()
+			.Padding(StandardPadding)
+			.FillWidth(0.5f)
+			[
+				SNew(STextBlock)
+				.Text(LOCTEXT("BrushRadius", "Brush Radius"))
+				.Font(LabelFont)
+			]
+			+ SHorizontalBox::Slot()
+			.Padding(StandardPadding)
+			.FillWidth(0.5f)
+			[
+				SNew(SNumericEntryBox<float>)
+				.Font(StandardFont)
+				.AllowSpin(true)
+				.MinValue(0.0f)
+				.MaxValue(65536.0f)
+				.MaxSliderValue(4096.0f)
+				.MinDesiredValueWidth(50.0f)
+				.SliderExponent(3.0f)
+				.Value_Lambda([this]
+				{
+					return BrushTool ? BrushTool->BrushRadius : 0;
+				})
+				.OnValueChanged_Lambda([this](float Value)
+				{
+					if(BrushTool)
+					{
+						BrushTool->BrushRadius = Value;
+					}
+				})
+			]
+		]
+		HEADER("Cleanup")
+		+ SVerticalBox::Slot()
+		.AutoHeight()
+		[
+			SNew(SHorizontalBox)
+			BUTTON_POINTS("CleanupHideSelected", "Hide Selected", {FLidarPointCloudEditorHelper::HideSelected();})
+			BUTTON_POINTS("CleanupResetVisibility", "Reset Visibility", {FLidarPointCloudEditorHelper::ResetVisibility();})
+		]
+		+ SVerticalBox::Slot()
+		.AutoHeight()
+		[
+			SNew(SHorizontalBox)
+			BUTTON_POINTS("CleanupDeleteSelected", "Delete Selected", {FLidarPointCloudEditorHelper::DeleteSelected();})
+			BUTTON("CleanupDeleteHidden", "Delete Hidden", {FLidarPointCloudEditorHelper::DeleteHidden();})
+		]
+		+ SVerticalBox::Slot()
+		.AutoHeight()
+		[
+			SNew(SHorizontalBox)
+			BUTTON_POINTS("CleanupCropToSelection", "Crop To Selection",
+			{
+				FLidarPointCloudEditorHelper::InvertSelection();
+				FLidarPointCloudEditorHelper::DeleteSelected();
+			})
+			+ SHorizontalBox::Slot()
+			.Padding(StandardPadding)
+			.FillWidth(0.5f)
+			[
+				SNew(SSpacer)
+			]
+		]
+		HEADER("Collisions")
+		+ SVerticalBox::Slot()
+		.AutoHeight()
+		[
+			SNew(SHorizontalBox)
+			+ SHorizontalBox::Slot()
+			.Padding(StandardPadding)
+			.FillWidth(0.5f)
+			[
+				SNew(STextBlock)
+				.Text(LOCTEXT("CollisionsError", "Max Error"))
+				.Font(LabelFont)
+			]
+			+ SHorizontalBox::Slot()
+			.Padding(StandardPadding)
+			.FillWidth(0.5f)
+			[
+				SNew(SNumericEntryBox<float>)
+				.Font(StandardFont)
+				.AllowSpin(true)
+				.MinValue(0.0f)
+				.MaxValue(65536.0f)
+				.MaxSliderValue(512.0f)
+				.MinDesiredValueWidth(50.0f)
+				.Value_Lambda([this]{ return MaxCollisionError; })
+				.OnValueChanged_Lambda([this](float Value) { MaxCollisionError = Value; })
+				.IsEnabled(this, &SLidarEditorWidget::IsActorSelection)
+			]
+		]
+		+ SVerticalBox::Slot()
+		.AutoHeight()
+		[
+			SNew(SHorizontalBox)
+			BUTTON_ACTORS("CollisionsAdd", "Add Collision",
+			{
+				FLidarPointCloudEditorHelper::SetCollisionErrorForSelection(MaxCollisionError);
+				FLidarPointCloudEditorHelper::BuildCollisionForSelection();
+			})
+			BUTTON_ACTORS("CollisionsRemove", "Remove Collision", {FLidarPointCloudEditorHelper::RemoveCollisionForSelection();})
+		]
+		HEADER("Normals")
+		+ SVerticalBox::Slot()
+		.AutoHeight()
+		[
+			SNew(SHorizontalBox)
+			+ SHorizontalBox::Slot()
+			.Padding(StandardPadding)
+			.FillWidth(0.5f)
+			[
+				SNew(STextBlock)
+				.Text(LOCTEXT("NormalsQuality", "Quality"))
+				.Font(LabelFont)
+			]
+			+ SHorizontalBox::Slot()
+			.Padding(StandardPadding)
+			.FillWidth(0.5f)
+			[
+				SNew(SNumericEntryBox<int32>)
+				.Font(StandardFont)
+				.AllowSpin(true)
+				.MinValue(0)
+				.MaxValue(100)
+				.Value_Lambda([this]{ return NormalsQuality; })
+				.OnValueChanged_Lambda([this](int32 Value) { NormalsQuality = Value; })
+			]
+		]
+		+ SVerticalBox::Slot()
+		.AutoHeight()
+		[
+			SNew(SHorizontalBox)
+			+ SHorizontalBox::Slot()
+			.Padding(StandardPadding)
+			.FillWidth(0.5f)
+			[
+				SNew(STextBlock)
+				.Text(LOCTEXT("NormalsNoiseTolerance", "Noise Tolerance"))
+				.Font(LabelFont)
+			]
+			+ SHorizontalBox::Slot()
+			.Padding(StandardPadding)
+			.FillWidth(0.5f)
+			[
+				SNew(SNumericEntryBox<float>)
+				.Font(StandardFont)
+				.AllowSpin(true)
+				.MinValue(0.0f)
+				.MaxValue(20.0f)
+				.MaxSliderValue(5.0f)
+				.MinDesiredValueWidth(50.0f)
+				.Value_Lambda([this]{ return NormalsNoiseTolerance; })
+				.OnValueChanged_Lambda([this](float Value) { NormalsNoiseTolerance = Value; })
+			]
+		]
+		+ SVerticalBox::Slot()
+		.AutoHeight()
+		[
+			SNew(SHorizontalBox)
+			BUTTON("NormalsCalculate", "Calculate Normals",
+			{
+				FLidarPointCloudEditorHelper::SetNormalsQuality(NormalsQuality, NormalsNoiseTolerance);
+				if(IsActorSelection())
+				{
+					FLidarPointCloudEditorHelper::CalculateNormals();
+				}
+				else
+				{
+					FLidarPointCloudEditorHelper::CalculateNormalsForSelection();
+				}
+			})
+		]
+		HEADER("Meshing")
+		+ SVerticalBox::Slot()
+		.AutoHeight()
+		[
+			SNew(SHorizontalBox)
+			+ SHorizontalBox::Slot()
+			.Padding(StandardPadding)
+			.FillWidth(0.5f)
+			[
+				SNew(STextBlock)
+				.Text(LOCTEXT("MeshingError", "Max Error"))
+				.Font(LabelFont)
+			]
+			+ SHorizontalBox::Slot()
+			.Padding(StandardPadding)
+			.FillWidth(0.5f)
+			[
+				SNew(SNumericEntryBox<float>)
+				.Font(StandardFont)
+				.AllowSpin(true)
+				.MinValue(0.0f)
+				.MaxValue(65536.0f)
+				.MaxSliderValue(512.0f)
+				.MinDesiredValueWidth(50.0f)
+				.Value_Lambda([this]{ return MaxMeshingError; })
+				.OnValueChanged_Lambda([this](float Value) { MaxMeshingError = Value; })
+			]
+		]
+		+ SVerticalBox::Slot()
+		.AutoHeight()
+		[
+			SNew(SHorizontalBox)
+			+ SHorizontalBox::Slot()
+			.Padding(StandardPadding)
+			.FillWidth(0.5f)
+			[
+				SNew(STextBlock)
+				.Text(LOCTEXT("MeshingMerge", "Merge Meshes"))
+				.Font(LabelFont)
+			]
+			+ SHorizontalBox::Slot()
+			.Padding(StandardPadding)
+			.FillWidth(0.5f)
+			[
+				SNew(SCheckBox)
+				.IsChecked_Lambda([this]{ return bMergeMeshes ? ECheckBoxState::Checked : ECheckBoxState::Unchecked; })
+				.OnCheckStateChanged_Lambda([this](ECheckBoxState Value) { bMergeMeshes = Value == ECheckBoxState::Checked; })
+			]
+		]
+		+ SVerticalBox::Slot()
+		.AutoHeight()
+		[
+			SNew(SHorizontalBox)
+			+ SHorizontalBox::Slot()
+			.Padding(StandardPadding)
+			.FillWidth(0.5f)
+			[
+				SNew(STextBlock)
+				.Text(LOCTEXT("MeshingRetain", "Retain Transform"))
+				.Font(LabelFont)
+			]
+			+ SHorizontalBox::Slot()
+			.Padding(StandardPadding)
+			.FillWidth(0.5f)
+			[
+				SNew(SCheckBox)
+				.IsChecked_Lambda([this]{ return !bMergeMeshes && bRetainTransform ? ECheckBoxState::Checked : ECheckBoxState::Unchecked; })
+				.OnCheckStateChanged_Lambda([this](ECheckBoxState Value) { bRetainTransform = Value == ECheckBoxState::Checked; })
+				.IsEnabled_Lambda([this] { return !bMergeMeshes; })
+			]
+		]
+		+ SVerticalBox::Slot()
+		.AutoHeight()
+		[
+			SNew(SHorizontalBox)
+			BUTTON("MeshingBuild", "Create Static Mesh", {FLidarPointCloudEditorHelper::MeshSelected(IsPointSelection(), MaxMeshingError, bMergeMeshes, !bMergeMeshes && bRetainTransform);})
+		]
+		HEADER("Alignment")
+		+ SVerticalBox::Slot()
+		.AutoHeight()
+		[
+			SNew(SHorizontalBox)
+			BUTTON_ACTORS("AlignmentCoords", "Original Coordinates", {FLidarPointCloudEditorHelper::SetOriginalCoordinateForSelection();})
+			BUTTON_ACTORS("AlignmentOrigin", "World Origin", {FLidarPointCloudEditorHelper::AlignSelectionAroundWorldOrigin();})
+		]
+		+ SVerticalBox::Slot()
+		.AutoHeight()
+		[
+			SNew(SHorizontalBox)
+			BUTTON_ACTORS("AlignmentReset", "Reset Alignment", {FLidarPointCloudEditorHelper::CenterSelection();})
+			+ SHorizontalBox::Slot()
+			.Padding(StandardPadding)
+			.FillWidth(0.5f)
+			[
+				SNew(SSpacer)
+			]
+		]
+		HEADEREX("Merge & Extract", "MergeExtract")
+		+ SVerticalBox::Slot()
+		.AutoHeight()
+		[
+			SNew(SHorizontalBox)
+			BUTTON_POINTS("MergeExtractSelection", "Extract Selection", {FLidarPointCloudEditorHelper::Extract();})
+			BUTTON_POINTS("MergeExtractCopy", "Extract as Copy", {FLidarPointCloudEditorHelper::ExtractAsCopy();})
+		]
+		+ SVerticalBox::Slot()
+		.AutoHeight()
+		[
+			SNew(SHorizontalBox)
+			BUTTON_ACTORS("MergeExtractActors", "Merge Actors", {FLidarPointCloudEditorHelper::MergeSelectionByComponent(true);})
+			BUTTON_ACTORS("MergeExtractData", "Merge Data", {FLidarPointCloudEditorHelper::MergeSelectionByData(true);})
+		]
+	];
+
+#undef BUTTON_ACTORS
+#undef BUTTON_POINTS
+#undef BUTTON
+#undef HEADER
+#undef HEADEREX
+}
+END_SLATE_FUNCTION_BUILD_OPTIMIZATION
+
 FLidarPointCloudEdModeToolkit::~FLidarPointCloudEdModeToolkit()
 {
 	GetScriptableEditorMode()->GetInteractiveToolsContext(EToolsContextScope::EdMode)->OnToolNotificationMessage.RemoveAll(this);
@@ -21,6 +392,8 @@ FLidarPointCloudEdModeToolkit::~FLidarPointCloudEdModeToolkit()
 
 void FLidarPointCloudEdModeToolkit::Init(const TSharedPtr<IToolkitHost>& InitToolkitHost, TWeakObjectPtr<UEdMode> InOwningMode)
 {
+	EditorWidget = SNew(SLidarEditorWidget);
+	
 	FModeToolkit::Init(InitToolkitHost, InOwningMode);
 	
 	GetScriptableEditorMode()->GetInteractiveToolsContext(EToolsContextScope::EdMode)->OnToolNotificationMessage.AddSP(this, &FLidarPointCloudEdModeToolkit::SetActiveToolMessage);
@@ -40,20 +413,11 @@ FText FLidarPointCloudEdModeToolkit::GetBaseToolkitName() const
 void FLidarPointCloudEdModeToolkit::GetToolPaletteNames(TArray<FName>& InPaletteName) const
 {
 	InPaletteName.Add(LidarEditorPalletes::Manage);
-	InPaletteName.Add(LidarEditorPalletes::Edit);
 }
 
 FText FLidarPointCloudEdModeToolkit::GetToolPaletteDisplayName(FName PaletteName) const
 {
-	if (PaletteName == LidarEditorPalletes::Manage)
-	{
-		return LOCTEXT("LidarMode_Manage", "Manage");
-	}
-	if (PaletteName == LidarEditorPalletes::Edit)
-	{
-		return LOCTEXT("LidarMode_Edit", "Edit");
-	}
-	return FText();
+	return LOCTEXT("LidarMode_Manage", "Manage");
 }
 
 FText FLidarPointCloudEdModeToolkit::GetActiveToolDisplayName() const
@@ -85,5 +449,20 @@ void FLidarPointCloudEdModeToolkit::SetActiveToolMessage(const FText& Message)
 	ActiveToolMessageHandle.Reset();
 }
 
+void FLidarPointCloudEdModeToolkit::SetActorSelection(bool bNewActorSelection)
+{
+	if(EditorWidget.IsValid())
+	{
+		EditorWidget->bActorSelection = bNewActorSelection;
+	}
+}
+
+void FLidarPointCloudEdModeToolkit::SetBrushTool(ULidarEditorToolPaintSelection* NewBrushTool)
+{
+	if(EditorWidget.IsValid())
+	{
+		EditorWidget->BrushTool = NewBrushTool;
+	}
+}
 
 #undef LOCTEXT_NAMESPACE
diff --git a/Engine/Plugins/Enterprise/LidarPointCloud/Source/LidarPointCloudEditor/Private/LidarPointCloudEdModeToolkit.h b/Engine/Plugins/Enterprise/LidarPointCloud/Source/LidarPointCloudEditor/Private/LidarPointCloudEdModeToolkit.h
--- a/Engine/Plugins/Enterprise/LidarPointCloud/Source/LidarPointCloudEditor/Private/LidarPointCloudEdModeToolkit.h
+++ b/Engine/Plugins/Enterprise/LidarPointCloud/Source/LidarPointCloudEditor/Private/LidarPointCloudEdModeToolkit.h
@@ -5,12 +5,46 @@
 #include "EditorModeManager.h"
 #include "StatusBarSubsystem.h"
 
+class ULidarEditorMode;
+class ULidarEditorToolPaintSelection;
+
 namespace LidarEditorPalletes
 {
 	static const FName Manage(TEXT("ToolMode_Manage")); 
-	static const FName Edit(TEXT("ToolMode_Edit"));
 }
 
+class SLidarEditorWidget : public SCompoundWidget
+{
+public:
+	SLATE_BEGIN_ARGS(SLidarEditorWidget) {}
+	SLATE_END_ARGS()
+
+public:
+	void Construct(const FArguments& InArgs);
+
+	bool IsActorSelection() const { return bActorSelection; }
+	bool IsPointSelection() const { return !bActorSelection; }
+	bool IsBrushVisible() const { return BrushTool != nullptr; }
+
+	EVisibility GetActorVisibility() const { return bActorSelection ? EVisibility::Visible : EVisibility::Collapsed; }
+	EVisibility GetPointVisibility() const { return bActorSelection ? EVisibility::Collapsed : EVisibility::Visible; }
+	EVisibility GetBrushVisibility() const { return IsBrushVisible() ? EVisibility::Visible : EVisibility::Collapsed; }
+
+private:	
+	float MaxCollisionError = 0.0f;
+	float MaxMeshingError = 0.0f;
+	bool bMergeMeshes = true;
+	bool bRetainTransform = true;
+	int32 NormalsQuality = 40;
+	float NormalsNoiseTolerance = 1.0f;
+
+	ULidarEditorMode* LidarEditorMode = nullptr;
+	ULidarEditorToolPaintSelection* BrushTool = nullptr;
+	bool bActorSelection = false;
+
+	friend class FLidarPointCloudEdModeToolkit;
+};
+
 /**
  * Public interface to Lidar Edit mode.
  */
@@ -31,10 +65,14 @@ public:
 	virtual bool HasIntegratedToolPalettes() const override { return true; }	
 	virtual FText GetActiveToolDisplayName() const override;
 	virtual FText GetActiveToolMessage() const override;
-
+	virtual TSharedPtr<SWidget> GetInlineContent() const override { return EditorWidget; }
+	
 	void SetActiveToolMessage(const FText& Message);
+	void SetActorSelection(bool bNewActorSelection);
+	void SetBrushTool(ULidarEditorToolPaintSelection* NewBrushTool);
 
 private:
 	FStatusBarMessageHandle ActiveToolMessageHandle;
 	FText ActiveToolMessageCache;
+	TSharedPtr<SLidarEditorWidget> EditorWidget;
 };
\ No newline at end of file
diff --git a/Engine/Plugins/Enterprise/LidarPointCloud/Source/LidarPointCloudEditor/Private/LidarPointCloudEditorCommands.cpp b/Engine/Plugins/Enterprise/LidarPointCloud/Source/LidarPointCloudEditor/Private/LidarPointCloudEditorCommands.cpp
--- a/Engine/Plugins/Enterprise/LidarPointCloud/Source/LidarPointCloudEditor/Private/LidarPointCloudEditorCommands.cpp
+++ b/Engine/Plugins/Enterprise/LidarPointCloud/Source/LidarPointCloudEditor/Private/LidarPointCloudEditorCommands.cpp
@@ -19,24 +19,10 @@ void FLidarPointCloudEditorCommands::RegisterCommands()
 	UI_COMMAND(Center, "Center", "Enable, to center the point cloud asset\nDisable, to use original coordinates.", EUserInterfaceActionType::ToggleButton, FInputChord());
 
 	// TOOLKIT COMMANDS
-
-	TArray<TSharedPtr<FUICommandInfo, ESPMode::ThreadSafe>>& ManageCommands = Commands.FindOrAdd(LidarEditorPalletes::Manage);
-	UI_COMMAND(ToolkitSelect, "Select", "Select point cloud assets", EUserInterfaceActionType::ToggleButton, FInputChord());
-	ManageCommands.Add(ToolkitSelect);
-	UI_COMMAND(ToolkitAlign, "Align", "Aligns selected point cloud assets", EUserInterfaceActionType::ToggleButton, FInputChord());
-	ManageCommands.Add(ToolkitAlign);
-	UI_COMMAND(ToolkitMerge, "Merge", "Merges selected point cloud assets", EUserInterfaceActionType::ToggleButton, FInputChord());
-	ManageCommands.Add(ToolkitMerge);
-	UI_COMMAND(ToolkitCollision, "Collision", "Create and Destroy point cloud asset collisions", EUserInterfaceActionType::ToggleButton, FInputChord());
-	ManageCommands.Add(ToolkitCollision);
-	UI_COMMAND(ToolkitNormals, "Normals", "Calculates normals", EUserInterfaceActionType::ToggleButton, FInputChord());
-	ManageCommands.Add(ToolkitNormals);
-	UI_COMMAND(ToolkitMeshing, "Meshing", "Build Static Meshes", EUserInterfaceActionType::ToggleButton, FInputChord());
-	ManageCommands.Add(ToolkitMeshing);
-	
-	TArray<TSharedPtr<FUICommandInfo, ESPMode::ThreadSafe>>& EditCommands = Commands.FindOrAdd(LidarEditorPalletes::Edit);
+	TArray<TSharedPtr<FUICommandInfo, ESPMode::ThreadSafe>>& EditCommands = ToolkitCommands.FindOrAdd(LidarEditorPalletes::Manage);
+	UI_COMMAND(ToolkitSelect, "Object", "Select point cloud assets", EUserInterfaceActionType::ToggleButton, FInputChord());
 	EditCommands.Add(ToolkitSelect);
-	UI_COMMAND(ToolkitBoxSelection, "Box", "Uses box to select points.", EUserInterfaceActionType::ToggleButton, FInputChord());
+	UI_COMMAND(ToolkitBoxSelection, "Marquee", "Uses box to select points.", EUserInterfaceActionType::ToggleButton, FInputChord());
 	EditCommands.Add(ToolkitBoxSelection);
 	UI_COMMAND(ToolkitPolygonalSelection, "Polygon", "Uses custom polygon to select points.", EUserInterfaceActionType::ToggleButton, FInputChord());
 	EditCommands.Add(ToolkitPolygonalSelection);
diff --git a/Engine/Plugins/Enterprise/LidarPointCloud/Source/LidarPointCloudEditor/Private/LidarPointCloudEditorCommands.h b/Engine/Plugins/Enterprise/LidarPointCloud/Source/LidarPointCloudEditor/Private/LidarPointCloudEditorCommands.h
--- a/Engine/Plugins/Enterprise/LidarPointCloud/Source/LidarPointCloudEditor/Private/LidarPointCloudEditorCommands.h
+++ b/Engine/Plugins/Enterprise/LidarPointCloud/Source/LidarPointCloudEditor/Private/LidarPointCloudEditorCommands.h
@@ -24,9 +24,9 @@ public:
 	virtual void RegisterCommands() override;
 	// End of TCommand<> interface
 
-	static TMap<FName, TArray<TSharedPtr<FUICommandInfo>>> GetCommands()
+	static TMap<FName, TArray<TSharedPtr<FUICommandInfo>>> GetToolkitCommands()
 	{
-		return Get().Commands;
+		return Get().ToolkitCommands;
 	}
 
 public:
@@ -40,11 +40,6 @@ public:
 	// TOOLKIT COMMANDS
 
 	TSharedPtr<FUICommandInfo> ToolkitSelect;
-	TSharedPtr<FUICommandInfo> ToolkitAlign;
-	TSharedPtr<FUICommandInfo> ToolkitMerge;
-	TSharedPtr<FUICommandInfo> ToolkitCollision;
-	TSharedPtr<FUICommandInfo> ToolkitNormals;
-	TSharedPtr<FUICommandInfo> ToolkitMeshing;
 	TSharedPtr<FUICommandInfo> ToolkitBoxSelection;
 	TSharedPtr<FUICommandInfo> ToolkitPolygonalSelection;
 	TSharedPtr<FUICommandInfo> ToolkitLassoSelection;
@@ -52,5 +47,5 @@ public:
 	TSharedPtr<FUICommandInfo> ToolktitCancelSelection;
 
 private:
-	TMap<FName, TArray<TSharedPtr<FUICommandInfo>>> Commands;
+	TMap<FName, TArray<TSharedPtr<FUICommandInfo>>> ToolkitCommands;
 };
\ No newline at end of file
diff --git a/Engine/Plugins/Enterprise/LidarPointCloud/Source/LidarPointCloudEditor/Private/LidarPointCloudEditorHelper.cpp b/Engine/Plugins/Enterprise/LidarPointCloud/Source/LidarPointCloudEditor/Private/LidarPointCloudEditorHelper.cpp
--- a/Engine/Plugins/Enterprise/LidarPointCloud/Source/LidarPointCloudEditor/Private/LidarPointCloudEditorHelper.cpp
+++ b/Engine/Plugins/Enterprise/LidarPointCloud/Source/LidarPointCloudEditor/Private/LidarPointCloudEditorHelper.cpp
@@ -107,7 +107,7 @@ namespace
 
 		return Actors;
 	}
-
+	
 	int32 GetNumLidarActors()
 	{
 		int32 NumActors = 0;
@@ -599,9 +599,9 @@ void FLidarPointCloudEditorHelper::CalculateNormalsForSelection()
 	});
 }
 
-void FLidarPointCloudEditorHelper::SetNormalsQualityForSelection(int32 Quality, float NoiseTolerance)
+void FLidarPointCloudEditorHelper::SetNormalsQuality(int32 Quality, float NoiseTolerance)
 {
-	ProcessSelection([Quality, NoiseTolerance](ALidarPointCloudActor* Actor)
+	ProcessAll([Quality, NoiseTolerance](ALidarPointCloudActor* Actor)
 	{
 		if(ULidarPointCloud* PointCloud = Actor->GetPointCloud())
 		{
@@ -862,6 +862,21 @@ void FLidarPointCloudEditorHelper::ClearSelection()
 	});
 }
 
+void FLidarPointCloudEditorHelper::InvertActorSelection()
+{
+	TArray<ALidarPointCloudActor*> SelectedActors = GetSelectedActors();
+
+	ProcessAll([&SelectedActors](ALidarPointCloudActor* Actor)
+	{
+		GEditor->SelectActor(Actor, !SelectedActors.Contains(Actor), true);
+	});
+}
+
+void FLidarPointCloudEditorHelper::ClearActorSelection()
+{
+	GEditor->SelectNone(true, true);
+}
+
 void FLidarPointCloudEditorHelper::MergeLidar(ULidarPointCloud* TargetAsset, TArray<ULidarPointCloud*> SourceAssets)
 {
 	if(!IsValid(TargetAsset) || SourceAssets.Num() == 0)
diff --git a/Engine/Plugins/Enterprise/LidarPointCloud/Source/LidarPointCloudEditor/Private/LidarPointCloudEditorHelper.h b/Engine/Plugins/Enterprise/LidarPointCloud/Source/LidarPointCloudEditor/Private/LidarPointCloudEditorHelper.h
--- a/Engine/Plugins/Enterprise/LidarPointCloud/Source/LidarPointCloudEditor/Private/LidarPointCloudEditorHelper.h
+++ b/Engine/Plugins/Enterprise/LidarPointCloud/Source/LidarPointCloudEditor/Private/LidarPointCloudEditorHelper.h
@@ -37,7 +37,7 @@ public:
 	static void MeshSelected(bool bMeshByPoints, float CellSize, bool bMergeMeshes, bool bRetainTransform);
 	
 	static void CalculateNormalsForSelection();
-	static void SetNormalsQualityForSelection(int32 Quality, float NoiseTolerance);
+	static void SetNormalsQuality(int32 Quality, float NoiseTolerance);
 	
 	static void SelectPointsByConvexVolume(const FConvexVolume& ConvexVolume, ELidarPointCloudSelectionMode SelectionMode);
 	static void SelectPointsBySphere(FSphere Sphere, ELidarPointCloudSelectionMode SelectionMode);
@@ -45,6 +45,9 @@ public:
 	static void DeleteSelected();
 	static void InvertSelection();
 	static void ClearSelection();
+	
+	static void InvertActorSelection();
+	static void ClearActorSelection();
 
 	static void ResetVisibility();
 	static void DeleteHidden();
diff --git a/Engine/Plugins/Enterprise/LidarPointCloud/Source/LidarPointCloudEditor/Private/LidarPointCloudEditorTools.cpp b/Engine/Plugins/Enterprise/LidarPointCloud/Source/LidarPointCloudEditor/Private/LidarPointCloudEditorTools.cpp
--- a/Engine/Plugins/Enterprise/LidarPointCloud/Source/LidarPointCloudEditor/Private/LidarPointCloudEditorTools.cpp
+++ b/Engine/Plugins/Enterprise/LidarPointCloud/Source/LidarPointCloudEditor/Private/LidarPointCloudEditorTools.cpp
@@ -121,7 +121,7 @@ void ULidarToolActionsMeshing::BuildStaticMesh()
 
 void ULidarToolActionsNormals::CalculateNormals()
 {
-	FLidarPointCloudEditorHelper::SetNormalsQualityForSelection(Quality, NoiseTolerance);
+	FLidarPointCloudEditorHelper::SetNormalsQuality(Quality, NoiseTolerance);
 	FLidarPointCloudEditorHelper::CalculateNormalsForSelection();
 }
 
diff --git a/Engine/Plugins/Enterprise/LidarPointCloud/Source/LidarPointCloudEditor/Private/LidarPointCloudEditorTools.h b/Engine/Plugins/Enterprise/LidarPointCloud/Source/LidarPointCloudEditor/Private/LidarPointCloudEditorTools.h
--- a/Engine/Plugins/Enterprise/LidarPointCloud/Source/LidarPointCloudEditor/Private/LidarPointCloudEditorTools.h
+++ b/Engine/Plugins/Enterprise/LidarPointCloud/Source/LidarPointCloudEditor/Private/LidarPointCloudEditorTools.h
@@ -21,6 +21,7 @@ public:
 	virtual void Setup() override;
 	virtual TObjectPtr<UInteractiveToolPropertySet> CreateToolActions() { return nullptr; }
 	virtual FText GetToolMessage() const;
+	virtual bool IsActorSelection() const { return true; }
 
 	UPROPERTY()
 	TObjectPtr<UInteractiveToolPropertySet> ToolActions = nullptr;
@@ -49,6 +50,8 @@ public:
 	virtual void Setup() override;
 	virtual void Render(IToolsContextRenderAPI* RenderAPI) override;
 
+	virtual bool IsActorSelection() const override { return false; }
+
 	// IClickDragBehaviorTarget implementation
 	virtual FInputRayHit CanBeginClickDragSequence(const FInputDeviceRay& PressPos) override { return FInputRayHit(TNumericLimits<float>::Max()); }
 	virtual void OnClickPress(const FInputDeviceRay& PressPos) override {}
@@ -266,35 +269,44 @@ class ULidarToolActionsSelection : public UInteractiveToolPropertySet
 	GENERATED_BODY()
 
 public:
-	UFUNCTION(CallInEditor, Category = Visibility)
-	void HideSelected();
-
-	UFUNCTION(CallInEditor, Category = Visibility)
-	void ResetVisibility();
-
-	UFUNCTION(CallInEditor, Category = Selection)
-	void InvertSelection();
-
-	UFUNCTION(CallInEditor, Category = Selection)
+	UFUNCTION(CallInEditor, Category = Selection, meta = (DisplayName = "Clear"))
 	void ClearSelection();
+	
+	UFUNCTION(CallInEditor, Category = Selection, meta = (DisplayName = "Invert"))
+	void InvertSelection();
 
-	UFUNCTION(CallInEditor, Category = Deletion)
+	UFUNCTION(CallInEditor, Category = Cleanup)
 	void DeleteSelected();
 	
-	UFUNCTION(CallInEditor, Category = Deletion)
+	UFUNCTION(CallInEditor, Category = Cleanup)
 	void DeleteHidden();
+	
+	UFUNCTION(CallInEditor, Category = Cleanup)
+	void HideSelected();
 
-	UFUNCTION(CallInEditor, Category = Extraction)
-	void Extract();
-
-	UFUNCTION(CallInEditor, Category = Extraction)
-	void ExtractAsCopy();
+	UFUNCTION(CallInEditor, Category = Cleanup)
+	void ResetVisibility();
+	
+	/** Higher values will generally result in more accurate calculations, at the expense of time */
+	UPROPERTY(EditAnywhere, Category = "Normals", meta = (ClampMin = "1", ClampMax = "100"))
+	int32 Quality = 40;
 
+	/**
+	 * Higher values are less susceptible to noise, but will most likely lose finer details, especially around hard edges.
+	 * Lower values retain more detail, at the expense of time.
+	 * NOTE: setting this too low will cause visual artifacts and geometry holes in noisier datasets.
+	 */
+	UPROPERTY(EditAnywhere, BlueprintReadWrite, Category = "Normals", meta = (ClampMin = "0.0"))
+	float NoiseTolerance = 1;
+	
 	UFUNCTION(CallInEditor, Category = Normals)
 	void CalculateNormals();
+	
+	UFUNCTION(CallInEditor, Category = "Merge & Extract")
+	void Extract();
 
-	UFUNCTION(CallInEditor, Category = Meshing)
-	void BuildStaticMesh();
+	UFUNCTION(CallInEditor, Category = "Merge & Extract")
+	void ExtractAsCopy();
 	
 	/** Max error around the meshed areas. Leave at 0 for max quality */
 	UPROPERTY(EditAnywhere, Category = Meshing, meta = (UIMin = "0", UIMax = "2000", ClampMin = "0"))
@@ -306,6 +318,9 @@ public:
 	/** If not merging meshes, this will retain the transform of the original cloud */
 	UPROPERTY(EditAnywhere, Category = Meshing, meta=(EditCondition="!bMergeMeshes"))
 	bool bRetainTransform = true;
+
+	UFUNCTION(CallInEditor, Category = Meshing, meta = (DisplayName = "Create Static Mesh"))
+	void BuildStaticMesh();
 };
 
 UCLASS()
@@ -420,11 +435,13 @@ public:
 private:
 	void Paint();
 
+public:
+	float BrushRadius;
+	
 private:
 	FVector3f HitLocation;
 	float LastHitDistance;
 	bool bHasHit;
-	float BrushRadius;
 };
 
 UCLASS()
