Support for Raytracing in Lidar

diff --git a/Engine/Plugins/Enterprise/LidarPointCloud/Shaders/Private/LidarPointCloudVertexFactory.ush b/Engine/Plugins/Enterprise/LidarPointCloud/Shaders/Private/LidarPointCloudVertexFactory.ush
--- a/Engine/Plugins/Enterprise/LidarPointCloud/Shaders/Private/LidarPointCloudVertexFactory.ush
+++ b/Engine/Plugins/Enterprise/LidarPointCloud/Shaders/Private/LidarPointCloudVertexFactory.ush
@@ -26,7 +26,9 @@ float RootCellSize;
 float3 RootExtent;
 uint bUsePerPointScaling;
 float ReversedVirtualDepthMultiplier;
+#if !RAYHITGROUPSHADER
 Buffer<float> TreeBuffer;
+#endif
 
 // Needed for WPO calculations
 float3 ViewRightVector;
@@ -57,8 +59,6 @@ uint bUseClassificationAlpha;
 
 float4 ClassificationColors[32];
 
-Buffer<float> DataBuffer;
-
 /**
  * [Tx, Ty, Tz, Invert]
  * [Fx, Fy, Fz, Ex]
@@ -223,6 +223,7 @@ float CalcSpriteSize(float3 Position)
 {
 	float VD = VirtualDepth;
 
+#if !RAYHITGROUPSHADER
 	if (bUsePerPointScaling)
 	{
 		VD = 0;
@@ -262,6 +263,7 @@ float CalcSpriteSize(float3 Position)
 			NodeExtent *= 0.5f;
 		}
 	}
+#endif
 	
 	return RootCellSize / pow(2.0f, VD);
 }
@@ -284,13 +286,13 @@ float3 ProcessPosition(float3 Position, int VertexID)
 float3 LidarGetInstancePosition(uint VertexId)
 {
 	uint Idx = (VertexId / 4) * 5;
-	return float3(DataBuffer[Idx], DataBuffer[Idx + 1], DataBuffer[Idx + 2]);
+	return float3(LidarVF.VertexFetch_Buffer[Idx], LidarVF.VertexFetch_Buffer[Idx + 1], LidarVF.VertexFetch_Buffer[Idx + 2]);
 }
 
 float3 LidarGetInstanceNormal(uint VertexId)
 {
 	uint Idx = (VertexId / 4) * 5;
-	return CalcTangentToLocal(asuint(DataBuffer[Idx + 4]))[2];
+	return CalcTangentToLocal(asuint(LidarVF.VertexFetch_Buffer[Idx + 4]))[2];
 }
 
 FVertexFactoryIntermediates GetVertexFactoryIntermediates(FVertexFactoryInput Input)
@@ -313,13 +315,13 @@ FVertexFactoryIntermediates GetVertexFactoryIntermediates(FVertexFactoryInput In
 	else
 	{
 		uint Idx = (Input.VertexId / 4) * 5;
-		Position = float3(DataBuffer[Idx], DataBuffer[Idx + 1], DataBuffer[Idx + 2]);
+		Position = float3(LidarVF.VertexFetch_Buffer[Idx], LidarVF.VertexFetch_Buffer[Idx + 1], LidarVF.VertexFetch_Buffer[Idx + 2]);
 		
-		uint ColorData = asuint(DataBuffer[Idx + 3]);		
+		uint ColorData = asuint(LidarVF.VertexFetch_Buffer[Idx + 3]);		
 		Color = half3(0x000000FF & (ColorData >> 16), 0x000000FF & (ColorData >> 8), 0x000000FF & ColorData) * 0.003921568627451;
 		Intensity = (0x000000FF & (ColorData >> 24)) * 0.003921568627451;
 
-		MetaData = asuint(DataBuffer[Idx + 4]);
+		MetaData = asuint(LidarVF.VertexFetch_Buffer[Idx + 4]);
 	}
 
 	Intermediates.TangentToLocal = CalcTangentToLocal(MetaData);
@@ -540,5 +542,66 @@ uint VertexFactoryGetPrimitiveId(FVertexFactoryInterpolantsVSToPS Interpolants)
 	return 0;
 }
 
+#if RAYHITGROUPSHADER || COMPUTESHADER
+uint GetNumRayTracingDynamicMeshVerticesIndirect()
+{
+	return 0;
+}
+#endif
+
+FVertexFactoryInput LoadVertexFactoryInputFromIndices(uint TriangleIndex, int VertexIndex)
+{
+	FVertexFactoryInput Input = (FVertexFactoryInput)0;
+
+	Input.VertexId = TriangleIndex * 3 + VertexIndex;
+
+	/**
+	 * When using static buffers we duplicate the data for each of the 4 vertices of the plane so we can use the VertexId directly
+	 * When using dynamic buffers, each 4 vertices of the plane will share the same data element so we need to divide by 4
+	 *
+	 * The buffer stride is 20 bytes and data is held in floats, so we need to multiply the starting offset by 5
+	 */
+	uint VertexOffset = (bUseStaticBuffers ? Input.VertexId : floor(Input.VertexId / 4)) * 5;
+
+	// Helps minimize incorrect self-shadowing due to camera-facing planes 
+	float ZOffset = 0.3;
+		
+	Input.Position.x = LidarVF.VertexFetch_Buffer[VertexOffset + 0];
+	Input.Position.y = LidarVF.VertexFetch_Buffer[VertexOffset + 1];
+	Input.Position.z = LidarVF.VertexFetch_Buffer[VertexOffset + 2] - ZOffset;
+
+	uint ColorData = asuint(LidarVF.VertexFetch_Buffer[VertexOffset + 3]);
+	Input.Color = half4(0x000000FF & (ColorData >> 16), 0x000000FF & (ColorData >> 8), 0x000000FF & ColorData, 0x000000FF & (ColorData >> 24)) * 0.003921568627451;
+	Input.MetaData = asuint(LidarVF.VertexFetch_Buffer[VertexOffset + 4]);
+
+	return Input;
+}
+
+#if RAYHITGROUPSHADER
+FVertexFactoryInput LoadVertexFactoryInputForHGS(uint TriangleIndex, int VertexIndex)
+{
+	FVertexFactoryInput Input = LoadVertexFactoryInputFromIndices(TriangleIndex, VertexIndex);
+		
+#if VF_USE_PRIMITIVE_SCENE_DATA
+	VF_GPUSCENE_SET_INPUT_FOR_RT(Input, GetInstanceUserData(), 0U);
+#endif // VF_USE_PRIMITIVE_SCENE_DATA
+		
+	return Input;
+}
+#endif
+
+#if COMPUTESHADER
+
+FVertexFactoryInput LoadVertexFactoryInputForDynamicUpdate(uint TriangleIndex, int VertexIndex, uint PrimitiveId)
+{
+	FVertexFactoryInput Input = LoadVertexFactoryInputFromIndices(TriangleIndex, VertexIndex);
+		
+	VF_GPUSCENE_SET_INPUT_FOR_RT(Input, PrimitiveId, 0U);
+
+	return Input;
+}
+
+#endif
+
 #include "/Engine/Private/VertexFactoryDefaultInterface.ush"
 
diff --git a/Engine/Plugins/Enterprise/LidarPointCloud/Source/LidarPointCloudRuntime/Private/LidarPointCloudComponent.cpp b/Engine/Plugins/Enterprise/LidarPointCloud/Source/LidarPointCloudRuntime/Private/LidarPointCloudComponent.cpp
--- a/Engine/Plugins/Enterprise/LidarPointCloud/Source/LidarPointCloudRuntime/Private/LidarPointCloudComponent.cpp
+++ b/Engine/Plugins/Enterprise/LidarPointCloud/Source/LidarPointCloudRuntime/Private/LidarPointCloudComponent.cpp
@@ -344,9 +344,15 @@ void FLidarPointCloudComponentRenderParams::UpdateFromComponent(ULidarPointCloud
 	bDrawNodeBounds = Component->bDrawNodeBounds;
 	bShouldRenderFacingNormals = Component->ShouldRenderFacingNormals();
 	bUseFrustumCulling = Component->bUseFrustumCulling;
-
+	
 	ScalingMethod = Component->ScalingMethod;
 
+	// Per-point scaling is currently unavailable in Ray Tracing
+	if(ScalingMethod == ELidarPointCloudScalingMethod::PerPoint && GetDefault<ULidarPointCloudSettings>()->bEnableLidarRayTracing)
+	{
+		ScalingMethod = ELidarPointCloudScalingMethod::PerNodeAdaptive;
+	}
+
 	ColorSource = Component->ColorSource;
 	PointShape = Component->GetPointShape();
 
diff --git a/Engine/Plugins/Enterprise/LidarPointCloud/Source/LidarPointCloudRuntime/Private/LidarPointCloudLODManager.cpp b/Engine/Plugins/Enterprise/LidarPointCloud/Source/LidarPointCloudRuntime/Private/LidarPointCloudLODManager.cpp
--- a/Engine/Plugins/Enterprise/LidarPointCloud/Source/LidarPointCloudRuntime/Private/LidarPointCloudLODManager.cpp
+++ b/Engine/Plugins/Enterprise/LidarPointCloud/Source/LidarPointCloudRuntime/Private/LidarPointCloudLODManager.cpp
@@ -437,6 +437,8 @@ int64 FLidarPointCloudLODManager::ProcessLOD(const TArray<FLidarPointCloudLODMan
 			SelectionFilterParams.AddZeroed(DeltaSize);
 		}
 	}
+	
+	const bool bUseRayTracing = GetDefault<ULidarPointCloudSettings>()->bEnableLidarRayTracing;
 
 	uint32 TotalPointsSelected = 0;
 	int64 NewNumPointsInFrustum = 0;
@@ -477,7 +479,7 @@ int64 FLidarPointCloudLODManager::ProcessLOD(const TArray<FLidarPointCloudLODMan
 					SelectionParams.MinDepth = RegisteredProxy.ComponentRenderParams.MinDepth;
 					SelectionParams.MaxDepth = RegisteredProxy.ComponentRenderParams.MaxDepth;
 					SelectionParams.BoundsScale = RegisteredProxy.ComponentRenderParams.BoundsScale;
-					SelectionParams.bUseFrustumCulling = RegisteredProxy.ComponentRenderParams.bUseFrustumCulling;
+					SelectionParams.bUseFrustumCulling = RegisteredProxy.ComponentRenderParams.bUseFrustumCulling && !bUseRayTracing;
 					SelectionParams.ClippingVolumes = RegisteredProxy.ComponentRenderParams.bOwnedByEditor ? nullptr : &ClippingVolumes; // Ignore clipping if in editor viewport
 
 					// Append visible nodes
@@ -548,6 +550,7 @@ int64 FLidarPointCloudLODManager::ProcessLOD(const TArray<FLidarPointCloudLODMan
 				UpdateData.RootCellSize = RegisteredProxy.Octree->GetRootCellSize();
 				UpdateData.ClippingVolumes = ClippingVolumes;
 				UpdateData.bUseStaticBuffers = bUseStaticBuffers && !RegisteredProxy.Octree->IsOptimizedForDynamicData();
+				UpdateData.bUseRayTracing = bUseRayTracing;
 				UpdateData.RenderParams = RegisteredProxy.ComponentRenderParams;
 
 #if !(UE_BUILD_SHIPPING)
@@ -652,7 +655,7 @@ int64 FLidarPointCloudLODManager::ProcessLOD(const TArray<FLidarPointCloudLODMan
 					{
 						for (FLidarPointCloudProxyUpdateDataNode& Node : UpdateData.SelectedNodes)
 						{
-							if (Node.BuildDataCache(UpdateData.bUseStaticBuffers))
+							if (Node.BuildDataCache(UpdateData.bUseStaticBuffers, UpdateData.bUseRayTracing))
 							{
 								MaxPointsPerNode = FMath::Max(MaxPointsPerNode, (uint32)Node.NumVisiblePoints);
 							}
diff --git a/Engine/Plugins/Enterprise/LidarPointCloud/Source/LidarPointCloudRuntime/Private/LidarPointCloudOctree.cpp b/Engine/Plugins/Enterprise/LidarPointCloud/Source/LidarPointCloudRuntime/Private/LidarPointCloudOctree.cpp
--- a/Engine/Plugins/Enterprise/LidarPointCloud/Source/LidarPointCloudRuntime/Private/LidarPointCloudOctree.cpp
+++ b/Engine/Plugins/Enterprise/LidarPointCloud/Source/LidarPointCloudRuntime/Private/LidarPointCloudOctree.cpp
@@ -135,6 +135,7 @@ FLidarPointCloudOctreeNode::FLidarPointCloudOctreeNode(FLidarPointCloudOctree* T
 	, bHasData(false)
 	, DataCache(nullptr)
 	, VertexFactory(nullptr)
+	, RayTracingGeometry(nullptr)
 	, bRenderDataDirty(true)
 	, bHasDataPending(false)
 	, bCanReleaseData(true)
@@ -198,7 +199,7 @@ FLidarPointCloudPoint* FLidarPointCloudOctreeNode::GetPersistentData() const
 	return GetData();
 }
 
-bool FLidarPointCloudOctreeNode::BuildDataCache(bool bUseStaticBuffers)
+bool FLidarPointCloudOctreeNode::BuildDataCache(bool bUseStaticBuffers, bool bUseRayTracing)
 {	
 	// Only include nodes with available data
 	if (HasData() && GetNumVisiblePoints())
@@ -235,26 +236,41 @@ bool FLidarPointCloudOctreeNode::BuildDataCache(bool bUseStaticBuffers)
 			}
 		}
 
-		if (bRenderDataDirty)
+		if(bUseRayTracing)
 		{
-			if (DataCache.IsValid())
+			if(!RayTracingGeometry.IsValid())
+			{
+				RayTracingGeometry = MakeShareable(new FLidarPointCloudRayTracingGeometry());
+				bRenderDataDirty = true;
+			}
+		}
+		else
+		{
+			if(RayTracingGeometry.IsValid())
 			{
-				DataCache->Resize(GetNumVisiblePoints() * 5);
+				RayTracingGeometry->ReleaseResource();
+				RayTracingGeometry.Reset();
+				bRenderDataDirty = true;
+			}
+		}
 
-				uint8* StructuredBuffer = (uint8*)RHILockBuffer(DataCache->Buffer, 0, GetNumVisiblePoints() * sizeof(FLidarPointCloudPoint), RLM_WriteOnly);
-				for (FLidarPointCloudPoint* P = GetData(), *DataEnd = P + GetNumVisiblePoints(); P != DataEnd; ++P)
-				{
-					FMemory::Memcpy(StructuredBuffer, P, sizeof(FLidarPointCloudPoint));
-					StructuredBuffer += sizeof(FLidarPointCloudPoint);
-				}
-				RHIUnlockBuffer(DataCache->Buffer);
+		if (bRenderDataDirty)
+		{
+			if (DataCache.IsValid() && !DataCache->IsInitialized())
+			{
+				DataCache->Initialize(GetData(), GetNumVisiblePoints());
 			}
 
-			if (VertexFactory.IsValid())
+			if (VertexFactory.IsValid() && !VertexFactory->IsInitialized())
 			{
 				VertexFactory->Initialize(GetData(), GetNumVisiblePoints());
 			}
 
+			if(RayTracingGeometry.IsValid() && !RayTracingGeometry->IsInitialized())
+			{
+				RayTracingGeometry->Initialize(GetNumVisiblePoints());
+			}
+
 			bRenderDataDirty = false;
 		}
 
@@ -822,7 +838,8 @@ void FLidarPointCloudOctreeNode::ReleaseDataCache()
 {
 	if (VertexFactory.IsValid() || DataCache.IsValid())
 	{
-		ENQUEUE_RENDER_COMMAND(LidarPointCloudOctreeNode_ReleaseDataCache)([DataCache = DataCache, VertexFactory = VertexFactory](FRHICommandListImmediate& RHICmdList)
+		ENQUEUE_RENDER_COMMAND(LidarPointCloudOctreeNode_ReleaseDataCache)([
+			DataCache = DataCache, VertexFactory = VertexFactory, RayTracingGeometry = RayTracingGeometry](FRHICommandListImmediate& RHICmdList)
 		{
 			if (DataCache.IsValid())
 			{
@@ -833,10 +850,16 @@ void FLidarPointCloudOctreeNode::ReleaseDataCache()
 			{
 				VertexFactory->ReleaseResource();
 			}
+			
+			if(RayTracingGeometry.IsValid())
+			{
+				RayTracingGeometry->ReleaseResource();
+			}
 		});
 
 		DataCache.Reset();
 		VertexFactory.Reset();
+		RayTracingGeometry.Reset();
 	}
 }
 
diff --git a/Engine/Plugins/Enterprise/LidarPointCloud/Source/LidarPointCloudRuntime/Private/Rendering/LidarPointCloudRenderBuffers.cpp b/Engine/Plugins/Enterprise/LidarPointCloud/Source/LidarPointCloudRuntime/Private/Rendering/LidarPointCloudRenderBuffers.cpp
--- a/Engine/Plugins/Enterprise/LidarPointCloud/Source/LidarPointCloudRuntime/Private/Rendering/LidarPointCloudRenderBuffers.cpp
+++ b/Engine/Plugins/Enterprise/LidarPointCloud/Source/LidarPointCloudRuntime/Private/Rendering/LidarPointCloudRenderBuffers.cpp
@@ -13,6 +13,8 @@
 #include "Settings/EditorStyleSettings.h"
 #endif
 
+IMPLEMENT_GLOBAL_SHADER_PARAMETER_STRUCT(FLidarPointCloudVertexFactoryUniformShaderParameters, "LidarVF");
+
 #define BINDPARAM(Name) Name.Bind(ParameterMap, TEXT(#Name))
 #define SETPARAM(Name) if (Name.IsBound()) { ShaderBindings.Add(Name, UserData->Name); }
 #define SETSRVPARAM(Name) if(UserData->Name) { SETPARAM(Name) }
@@ -92,6 +94,10 @@ void FLidarPointCloudRenderBuffer::InitRHI()
 	FRHIResourceCreateInfo CreateInfo(TEXT("FLidarPointCloudRenderBuffer"));
 	Buffer = RHICreateVertexBuffer(sizeof(uint32) * Capacity, BUF_ShaderResource | BUF_Dynamic, CreateInfo);
 	SRV = RHICreateShaderResourceView(Buffer, sizeof(uint32), PF_R32_FLOAT);
+
+	FLidarPointCloudVertexFactoryUniformShaderParameters UniformParameters;
+	UniformParameters.VertexFetch_Buffer = SRV;
+	UniformBuffer = TUniformBufferRef<FLidarPointCloudVertexFactoryUniformShaderParameters>::CreateUniformBufferImmediate(UniformParameters, UniformBuffer_MultiFrame);
 }
 
 void FLidarPointCloudRenderBuffer::ReleaseRHI()
@@ -99,6 +105,11 @@ void FLidarPointCloudRenderBuffer::ReleaseRHI()
 	// This must be called from Rendering thread
 	check(IsInRenderingThread());
 
+	if(UniformBuffer.IsValid())
+	{
+		UniformBuffer.SafeRelease();
+	}
+	
 	if (Buffer)
 	{
 		Buffer.SafeRelease();
@@ -107,6 +118,47 @@ void FLidarPointCloudRenderBuffer::ReleaseRHI()
 	SRV.SafeRelease();
 }
 
+void FLidarPointCloudRenderBuffer::Initialize(FLidarPointCloudPoint* Data, int32 NumPoints)
+{
+	Resize(NumPoints * 5);
+
+	uint8* StructuredBuffer = (uint8*)RHILockBuffer(Buffer, 0, NumPoints * sizeof(FLidarPointCloudPoint), RLM_WriteOnly);
+	for (FLidarPointCloudPoint* P = Data, *DataEnd = P + NumPoints; P != DataEnd; ++P)
+	{
+		FMemory::Memcpy(StructuredBuffer, P, sizeof(FLidarPointCloudPoint));
+		StructuredBuffer += sizeof(FLidarPointCloudPoint);
+	}
+	RHIUnlockBuffer(Buffer);
+}
+
+//////////////////////////////////////////////////////////// Ray Tracing Geometry
+
+void FLidarPointCloudRayTracingGeometry::Initialize(int32 NumPoints)
+{
+#if RHI_RAYTRACING
+	NumPrimitives = NumPoints * 2;
+	NumVertices = NumPoints * 4;
+	
+	SetInitializer(FRayTracingGeometryInitializer());
+	
+	Initializer.IndexBuffer = GLidarPointCloudIndexBuffer.IndexBufferRHI;
+	Initializer.TotalPrimitiveCount = NumPrimitives;
+	Initializer.GeometryType = RTGT_Triangles;
+	Initializer.bFastBuild = true;
+	Initializer.bAllowUpdate = true;
+
+	InitResource();
+	
+	FRayTracingGeometrySegment Segment;
+	Segment.VertexBuffer = nullptr;
+	Segment.NumPrimitives = NumPrimitives;
+	Segment.MaxVertices = NumVertices;
+	Initializer.Segments.Add(Segment);
+	
+	UpdateRHI();
+#endif
+}
+
 //////////////////////////////////////////////////////////// User Data
 
 FLidarPointCloudBatchElementUserData::FLidarPointCloudBatchElementUserData()
@@ -141,7 +193,6 @@ void FLidarPointCloudBatchElementUserData::SetClassificationColors(const TMap<in
 void FLidarPointCloudVertexFactoryShaderParameters::Bind(const FShaderParameterMap& ParameterMap)
 {
 	BINDPARAM(TreeBuffer);
-	BINDPARAM(DataBuffer);
 	BINDPARAM(bEditorView);
 	BINDPARAM(SelectionColor);
 	BINDPARAM(LocationOffset);
@@ -179,16 +230,16 @@ void FLidarPointCloudVertexFactoryShaderParameters::Bind(const FShaderParameterM
 void FLidarPointCloudVertexFactoryShaderParameters::GetElementShaderBindings(const class FSceneInterface* Scene, const FSceneView* View, const FMeshMaterialShader* Shader, const EVertexInputStreamType InputStreamType, ERHIFeatureLevel::Type FeatureLevel,
 	const FVertexFactory* VertexFactory, const FMeshBatchElement& BatchElement, class FMeshDrawSingleShaderBindings& ShaderBindings, FVertexInputStreamArray& VertexStreams) const
 {
+	FRHIUniformBuffer* VertexFactoryUniformBuffer = static_cast<FRHIUniformBuffer*>(BatchElement.VertexFactoryUserData);
 	FLidarPointCloudBatchElementUserData* UserData = (FLidarPointCloudBatchElementUserData*)BatchElement.UserData;
 
 	SETSRVPARAM(TreeBuffer);
-	SETSRVPARAM(DataBuffer);
-	if (!UserData->DataBuffer && DataBuffer.IsBound())
+	
+	if(VertexFactoryUniformBuffer)
 	{
-		FRHIShaderResourceView* Ptr = GDummyLidarPointCloudRenderBuffer.SRV;
-		ShaderBindings.Add(DataBuffer, Ptr);
+		ShaderBindings.Add(Shader->GetUniformBufferParameter<FLidarPointCloudVertexFactoryUniformShaderParameters>(), VertexFactoryUniformBuffer);
 	}
-
+	
 	SETPARAM(bEditorView);
 	SETPARAM(SelectionColor);
 	SETPARAM(LocationOffset);
@@ -229,6 +280,33 @@ bool FLidarPointCloudVertexFactoryBase::ShouldCompilePermutation(const FVertexFa
 		Parameters.MaterialParameters.MaterialDomain == MD_Surface && Parameters.MaterialParameters.bIsUsedWithLidarPointCloud) || Parameters.MaterialParameters.bIsSpecialEngineMaterial;
 }
 
+void FLidarPointCloudVertexFactoryBase::ModifyCompilationEnvironment(
+	const FVertexFactoryShaderPermutationParameters& Parameters, FShaderCompilerEnvironment& OutEnvironment)
+{
+	OutEnvironment.SetDefine(TEXT("RAY_TRACING_DYNAMIC_MESH_IN_LOCAL_SPACE"), TEXT("1"));
+}
+
+FShaderResourceViewRHIRef FLidarPointCloudVertexFactoryBase::GetVertexBufferSRV()
+{
+	return GDummyLidarPointCloudRenderBuffer.SRV;
+}
+
+void FLidarPointCloudVertexFactoryBase::InitRHI()
+{
+	FLidarPointCloudVertexFactoryUniformShaderParameters UniformParameters;
+	UniformParameters.VertexFetch_Buffer = GetVertexBufferSRV();
+	if(UniformParameters.VertexFetch_Buffer)
+	{
+		UniformBuffer = TUniformBufferRef<FLidarPointCloudVertexFactoryUniformShaderParameters>::CreateUniformBufferImmediate(UniformParameters, UniformBuffer_MultiFrame);
+	}
+}
+
+void FLidarPointCloudVertexFactoryBase::ReleaseRHI()
+{
+	UniformBuffer.SafeRelease();
+	FVertexFactory::ReleaseRHI();
+}
+
 void FLidarPointCloudVertexFactory::Initialize(FLidarPointCloudPoint* Data, int32 NumPoints)
 {
 	if (IsInitialized())
@@ -246,16 +324,17 @@ void FLidarPointCloudVertexFactory::FPointCloudVertexBuffer::InitRHI()
 {
 	const uint32 BufferSize = NumPoints * 4 * sizeof(FLidarPointCloudPoint);
 	FRHIResourceCreateInfo CreateInfo(TEXT("FPointCloudVertexBuffer"));
-	VertexBufferRHI = RHICreateBuffer(BufferSize, BUF_Static | BUF_VertexBuffer, sizeof(FLidarPointCloudPoint), ERHIAccess::VertexOrIndexBuffer, CreateInfo);
+	VertexBufferRHI = RHICreateBuffer(BufferSize, BUF_ShaderResource | BUF_Static | BUF_VertexBuffer, sizeof(FLidarPointCloudPoint), ERHIAccess::VertexOrIndexBuffer, CreateInfo);
+	VertexBufferSRV = RHICreateShaderResourceView(VertexBufferRHI, sizeof(float), PF_R32_FLOAT);
 	void* Buffer = RHILockBuffer(VertexBufferRHI, 0, BufferSize, RLM_WriteOnly);
 
-	uint8* Dest = (uint8*)Buffer;
+	FLidarPointCloudPoint* Dest = (FLidarPointCloudPoint*)Buffer;
 	for (int32 i = 0; i < NumPoints; ++i, ++Data)
 	{
-		FMemory::Memcpy(Dest, Data, sizeof(FLidarPointCloudPoint)); Dest += sizeof(FLidarPointCloudPoint);
-		FMemory::Memcpy(Dest, Data, sizeof(FLidarPointCloudPoint)); Dest += sizeof(FLidarPointCloudPoint);
-		FMemory::Memcpy(Dest, Data, sizeof(FLidarPointCloudPoint)); Dest += sizeof(FLidarPointCloudPoint);
-		FMemory::Memcpy(Dest, Data, sizeof(FLidarPointCloudPoint)); Dest += sizeof(FLidarPointCloudPoint);
+		FMemory::Memcpy(Dest, Data, sizeof(FLidarPointCloudPoint)); Dest++;
+		FMemory::Memcpy(Dest, Data, sizeof(FLidarPointCloudPoint)); Dest++;
+		FMemory::Memcpy(Dest, Data, sizeof(FLidarPointCloudPoint)); Dest++;
+		FMemory::Memcpy(Dest, Data, sizeof(FLidarPointCloudPoint)); Dest++;
 	}
 
 	RHIUnlockBuffer(VertexBufferRHI);
@@ -264,17 +343,19 @@ void FLidarPointCloudVertexFactory::FPointCloudVertexBuffer::InitRHI()
 void FLidarPointCloudVertexFactory::InitRHI()
 {
 	VertexBuffer.InitResource();
-
+	
 	FVertexDeclarationElementList Elements;
 	Elements.Add(AccessStreamComponent(FVertexStreamComponent(&VertexBuffer, 0, sizeof(FLidarPointCloudPoint), VET_Float3), 0));
 	Elements.Add(AccessStreamComponent(FVertexStreamComponent(&VertexBuffer, 12, sizeof(FLidarPointCloudPoint), VET_Color), 1));
 	Elements.Add(AccessStreamComponent(FVertexStreamComponent(&VertexBuffer, 16, sizeof(FLidarPointCloudPoint), VET_UInt), 2));
 	InitDeclaration(Elements);
+
+	FLidarPointCloudVertexFactoryBase::InitRHI();
 }
 
 void FLidarPointCloudVertexFactory::ReleaseRHI()
 {
-	FVertexFactory::ReleaseRHI();
+	FLidarPointCloudVertexFactoryBase::ReleaseRHI();
 	VertexBuffer.ReleaseResource();
 }
 
@@ -296,22 +377,29 @@ void FLidarPointCloudSharedVertexFactory::InitRHI()
 	Elements.Add(AccessStreamComponent(FVertexStreamComponent(&VertexBuffer, 0, 0, VET_Color), 1));
 	Elements.Add(AccessStreamComponent(FVertexStreamComponent(&VertexBuffer, 0, 0, VET_Color), 2));
 	InitDeclaration(Elements);
+
+	FLidarPointCloudVertexFactoryBase::InitRHI();
 }
 
 void FLidarPointCloudSharedVertexFactory::ReleaseRHI()
 {
-	FVertexFactory::ReleaseRHI();
+	FLidarPointCloudVertexFactoryBase::ReleaseRHI();
 	VertexBuffer.ReleaseResource();
 }
 
 IMPLEMENT_TYPE_LAYOUT(FLidarPointCloudVertexFactoryShaderParameters);
 
 IMPLEMENT_VERTEX_FACTORY_PARAMETER_TYPE(FLidarPointCloudVertexFactoryBase, SF_Vertex, FLidarPointCloudVertexFactoryShaderParameters);
+#if RHI_RAYTRACING
+IMPLEMENT_VERTEX_FACTORY_PARAMETER_TYPE(FLidarPointCloudVertexFactoryBase, SF_Compute, FLidarPointCloudVertexFactoryShaderParameters);
+IMPLEMENT_VERTEX_FACTORY_PARAMETER_TYPE(FLidarPointCloudVertexFactoryBase, SF_RayHitGroup, FLidarPointCloudVertexFactoryShaderParameters);
+#endif
 
 IMPLEMENT_VERTEX_FACTORY_TYPE(FLidarPointCloudVertexFactoryBase, "/Plugin/LidarPointCloud/Private/LidarPointCloudVertexFactory.ush",
 	  EVertexFactoryFlags::UsedWithMaterials
 	| EVertexFactoryFlags::SupportsDynamicLighting
 	| EVertexFactoryFlags::SupportsPositionOnly
+	| EVertexFactoryFlags::SupportsRayTracingDynamicGeometry
 );
 
 #undef BINDPARAM
diff --git a/Engine/Plugins/Enterprise/LidarPointCloud/Source/LidarPointCloudRuntime/Private/Rendering/LidarPointCloudRenderBuffers.h b/Engine/Plugins/Enterprise/LidarPointCloud/Source/LidarPointCloudRuntime/Private/Rendering/LidarPointCloudRenderBuffers.h
--- a/Engine/Plugins/Enterprise/LidarPointCloud/Source/LidarPointCloudRuntime/Private/Rendering/LidarPointCloudRenderBuffers.h
+++ b/Engine/Plugins/Enterprise/LidarPointCloud/Source/LidarPointCloudRuntime/Private/Rendering/LidarPointCloudRenderBuffers.h
@@ -9,6 +9,12 @@
 #include "HAL/ThreadSafeBool.h"
 #include "MeshMaterialShader.h"
 
+class FLidarPointCloudVertexFactory;
+
+BEGIN_GLOBAL_SHADER_PARAMETER_STRUCT(FLidarPointCloudVertexFactoryUniformShaderParameters,)
+	SHADER_PARAMETER_SRV(Buffer<float>, VertexFetch_Buffer)
+END_GLOBAL_SHADER_PARAMETER_STRUCT()
+
 /**
  * Base class for the buffers
  */
@@ -50,6 +56,10 @@ public:
 	FLidarPointCloudRenderBuffer() : FLidarPointCloudBuffer() {}
 	FLidarPointCloudRenderBuffer(uint32 Capacity) : FLidarPointCloudBuffer(Capacity) {}
 
+	void Initialize(FLidarPointCloudPoint* Data, int32 NumPoints);
+	
+	virtual FRHIUniformBuffer* GetUniformBuffer() const { return UniformBuffer.GetReference(); }
+
 	FBufferRHIRef Buffer;
 	FShaderResourceViewRHIRef SRV;
 
@@ -58,6 +68,26 @@ public:
 	virtual void Resize(const uint32& RequestedCapacity) override;	
 	virtual void InitRHI() override;
 	virtual void ReleaseRHI() override;
+
+private:
+	TUniformBufferRef<FLidarPointCloudVertexFactoryUniformShaderParameters> UniformBuffer;
+};
+
+class FLidarPointCloudRayTracingGeometry : public FRayTracingGeometry
+{
+public:
+	FLidarPointCloudRayTracingGeometry() {}
+	virtual ~FLidarPointCloudRayTracingGeometry() {}
+	
+	void Initialize(int32 NumPoints);
+
+	uint32 GetNumVertices() const { return NumVertices; }
+	uint32 GetNumPrimitives() const { return NumPrimitives; }
+	uint32 GetBufferSize() const { return NumVertices * sizeof(FLidarPointCloudPoint); }
+
+private:
+	uint32 NumVertices;
+	uint32 NumPrimitives;
 };
 
 /**
@@ -66,7 +96,6 @@ public:
 struct FLidarPointCloudBatchElementUserData
 {
 	FRHIShaderResourceView* TreeBuffer;
-	FRHIShaderResourceView* DataBuffer;
 	int32 bEditorView;
 	FVector3f SelectionColor;
 	FVector3f LocationOffset;
@@ -118,7 +147,6 @@ public:
 		const FVertexFactory* VertexFactory, const FMeshBatchElement& BatchElement, class FMeshDrawSingleShaderBindings& ShaderBindings, FVertexInputStreamArray& VertexStreams) const;
 
 	LAYOUT_FIELD(FShaderResourceParameter, TreeBuffer);
-	LAYOUT_FIELD(FShaderResourceParameter, DataBuffer);
 	LAYOUT_FIELD(FShaderParameter, bEditorView);
 	LAYOUT_FIELD(FShaderParameter, SelectionColor);
 	LAYOUT_FIELD(FShaderParameter, LocationOffset);
@@ -156,24 +184,40 @@ public:
 class FLidarPointCloudVertexFactoryBase : public FVertexFactory
 {
 	DECLARE_VERTEX_FACTORY_TYPE(FLidarPointCloudVertexFactoryBase);
-
+	
 public:
 	static bool ShouldCache(const FVertexFactoryShaderPermutationParameters& Parameters) { return IsFeatureLevelSupported(Parameters.Platform, ERHIFeatureLevel::SM5); }
 	static bool ShouldCompilePermutation(const FVertexFactoryShaderPermutationParameters& Parameters);
+	static void ModifyCompilationEnvironment(const FVertexFactoryShaderPermutationParameters& Parameters, FShaderCompilerEnvironment& OutEnvironment);
 
 	FLidarPointCloudVertexFactoryBase() : FVertexFactory(ERHIFeatureLevel::SM5) { }
+
+	virtual FVertexBuffer* GetVertexBuffer() { return nullptr; }
+	virtual FBufferRHIRef GetVertexBufferRHI() { return GetVertexBuffer() ? GetVertexBuffer()->VertexBufferRHI : nullptr; }
+	virtual FShaderResourceViewRHIRef GetVertexBufferSRV();
+	virtual FRHIUniformBuffer* GetUniformBuffer() const { return UniformBuffer.GetReference(); }
+	
+	virtual void InitRHI() override;
+	virtual void ReleaseRHI() override;
+
+protected:
+	TUniformBufferRef<FLidarPointCloudVertexFactoryUniformShaderParameters> UniformBuffer;
 };
 
 class FLidarPointCloudVertexFactory : public FLidarPointCloudVertexFactoryBase
 {
 public:
 	void Initialize(FLidarPointCloudPoint* Data, int32 NumPoints);
+	
+	virtual FVertexBuffer* GetVertexBuffer() override { return &VertexBuffer; }
+	virtual FShaderResourceViewRHIRef GetVertexBufferSRV() override { return VertexBuffer.VertexBufferSRV; }
 
 private:
 	class FPointCloudVertexBuffer : public FVertexBuffer
 	{
 		FLidarPointCloudPoint* Data;
 		int32 NumPoints;
+		FShaderResourceViewRHIRef VertexBufferSRV;
 
 	public:
 		virtual void InitRHI() override;
@@ -190,14 +234,19 @@ private:
  */
 class FLidarPointCloudSharedVertexFactory : public FLidarPointCloudVertexFactoryBase
 {
+public:
+	virtual FVertexBuffer* GetVertexBuffer() override { return &VertexBuffer; }
+
+private:
 	/** Very simple implementation of a ZeroStride Vertex Buffer */
 	class FPointCloudVertexBuffer : public FVertexBuffer
 	{
+		FShaderResourceViewRHIRef VertexBufferSRV;
 	public:
 		virtual void InitRHI() override;
 		virtual FString GetFriendlyName() const override { return TEXT("FPointCloudVertexBuffer"); }
 	} VertexBuffer;
-
+	
 	virtual void InitRHI() override;
 	virtual void ReleaseRHI() override;
 };
diff --git a/Engine/Plugins/Enterprise/LidarPointCloud/Source/LidarPointCloudRuntime/Private/Rendering/LidarPointCloudRendering.cpp b/Engine/Plugins/Enterprise/LidarPointCloud/Source/LidarPointCloudRuntime/Private/Rendering/LidarPointCloudRendering.cpp
--- a/Engine/Plugins/Enterprise/LidarPointCloud/Source/LidarPointCloudRuntime/Private/Rendering/LidarPointCloudRendering.cpp
+++ b/Engine/Plugins/Enterprise/LidarPointCloud/Source/LidarPointCloudRuntime/Private/Rendering/LidarPointCloudRendering.cpp
@@ -13,15 +13,19 @@
 #include "SceneManagement.h"
 #include "LocalVertexFactory.h"
 #include "Materials/Material.h"
+#if RHI_RAYTRACING
+#include "RayTracingInstance.h"
+#endif
 
 DECLARE_DWORD_COUNTER_STAT(TEXT("Draw Calls"), STAT_DrawCallCount, STATGROUP_LidarPointCloud)
 
-bool FLidarPointCloudProxyUpdateDataNode::BuildDataCache(bool bUseStaticBuffers)
+bool FLidarPointCloudProxyUpdateDataNode::BuildDataCache(bool bUseStaticBuffers, bool bUseRayTracing)
 {
-	if(DataNode && DataNode->BuildDataCache(bUseStaticBuffers))
+	if(DataNode && DataNode->BuildDataCache(bUseStaticBuffers, bUseRayTracing))
 	{
 		VertexFactory = DataNode->GetVertexFactory();
 		DataCache = DataNode->GetDataCache();
+		RayTracingGeometry = DataNode->GetRayTracingGeometry();
 		DataNode = nullptr;
 		return  true;
 	}
@@ -202,8 +206,6 @@ public:
 			return;
 		}
 
-		const bool bUsesSprites = RenderData.RenderParams.PointSize > 0;
-
 		for (int32 ViewIndex = 0; ViewIndex < Views.Num(); ViewIndex++)
 		{
 			const FSceneView* View = Views[ViewIndex];
@@ -218,25 +220,12 @@ public:
 
 					for (const FLidarPointCloudProxyUpdateDataNode& Node : RenderData.SelectedNodes)
 					{
-						if ((RenderData.bUseStaticBuffers && Node.VertexFactory.IsValid() && Node.VertexFactory->IsInitialized()) || (!RenderData.bUseStaticBuffers && Node.DataCache.IsValid()))
+						if ((RenderData.bUseStaticBuffers && Node.VertexFactory.IsValid() && Node.VertexFactory->IsInitialized()) ||
+							(!RenderData.bUseStaticBuffers && Node.DataCache.IsValid()))
 						{
 							FMeshBatch& MeshBatch = Collector.AllocateMesh();
-
-							MeshBatch.Type = bUsesSprites ? PT_TriangleList : PT_PointList;
-							MeshBatch.LODIndex = 0;
-							MeshBatch.VertexFactory = RenderData.bUseStaticBuffers ? Node.VertexFactory.Get() : (FVertexFactory*)&GLidarPointCloudSharedVertexFactory;
-							MeshBatch.bWireframe = false;
-							MeshBatch.MaterialRenderProxy = RenderData.RenderParams.Material->GetRenderProxy();
-							MeshBatch.ReverseCulling = IsLocalToWorldDeterminantNegative();
-							MeshBatch.DepthPriorityGroup = SDPG_World;
-
-							FMeshBatchElement& BatchElement = MeshBatch.Elements[0];
-							BatchElement.PrimitiveUniformBuffer = GetUniformBuffer();
-							BatchElement.IndexBuffer = &GLidarPointCloudIndexBuffer;
-							BatchElement.FirstIndex = bUsesSprites ? 0 : GLidarPointCloudIndexBuffer.PointOffset;
-							BatchElement.MinVertexIndex = 0;
-							BatchElement.NumPrimitives = Node.NumVisiblePoints * (bUsesSprites ? 2 : 1);
-							BatchElement.UserData = &UserData[UserData.Add(BuildUserDataElement(View, Node))];
+							
+							SetupMeshBatch(MeshBatch, Node, &UserData[UserData.Add(BuildUserDataElement(View, Node))]);
 
 							Collector.AddMesh(ViewIndex, MeshBatch);
 
@@ -309,6 +298,7 @@ public:
 			Result.bUsesLightingChannels = GetLightingChannelMask() != GetDefaultLightingChannelMask();
 			Result.bRenderCustomDepth = ShouldRenderCustomDepth();
 			MaterialRelevance.SetPrimitiveViewRelevance(Result);
+			Result.bVelocityRelevance = DrawsVelocity() && Result.bOpaque && Result.bRenderInMainPass;
 		}
 
 		return Result;
@@ -333,8 +323,8 @@ public:
 		UserDataElement.RootExtent = FVector3f(RenderData.RenderParams.BoundsSize.GetAbsMax() * 0.5f);
 
 		UserDataElement.LocationOffset = RenderData.RenderParams.LocationOffset;
-		UserDataElement.ViewRightVector = (FVector3f)InView->GetViewRight();
-		UserDataElement.ViewUpVector = (FVector3f)InView->GetViewUp();
+		UserDataElement.ViewRightVector = InView ? (FVector3f)InView->GetViewRight() : FVector3f::RightVector;
+		UserDataElement.ViewUpVector = InView ? (FVector3f)InView->GetViewUp() : FVector3f::ForwardVector;
 		UserDataElement.bUseCameraFacing = !RenderData.RenderParams.bShouldRenderFacingNormals;
 		UserDataElement.BoundsSize = RenderData.RenderParams.BoundsSize;
 		UserDataElement.ElevationColorBottom = FVector3f(RenderData.RenderParams.ColorSource == ELidarPointCloudColorationMode::None ? FColor::White : RenderData.RenderParams.ElevationColorBottom);
@@ -363,7 +353,6 @@ public:
 			UserDataElement.bStartClipped |= ClippingVolume.Mode == ELidarClippingVolumeMode::ClipOutside;
 		}
 
-		UserDataElement.DataBuffer = Node.DataCache ? Node.DataCache->SRV : nullptr;
 		UserDataElement.TreeBuffer = TreeBuffer->SRV;
 
 		return UserDataElement;
@@ -395,6 +384,79 @@ public:
 		RHIUnlockBuffer(TreeBuffer->Buffer);
 	}
 
+	void SetupMeshBatch(FMeshBatch& MeshBatch, const FLidarPointCloudProxyUpdateDataNode& Node, const FLidarPointCloudBatchElementUserData* UserData) const
+	{
+		const bool bUsesSprites = RenderData.RenderParams.PointSize > 0;
+		
+		MeshBatch.Type = bUsesSprites ? PT_TriangleList : PT_PointList;
+		MeshBatch.LODIndex = 0;
+		MeshBatch.VertexFactory = RenderData.bUseStaticBuffers ? Node.VertexFactory.Get() : (FVertexFactory*)&GLidarPointCloudSharedVertexFactory;
+		MeshBatch.bWireframe = false;
+		MeshBatch.MaterialRenderProxy = RenderData.RenderParams.Material->GetRenderProxy();
+		MeshBatch.ReverseCulling = IsLocalToWorldDeterminantNegative();
+		MeshBatch.DepthPriorityGroup = SDPG_World;
+		MeshBatch.bCanApplyViewModeOverrides = true;
+		
+		FMeshBatchElement& BatchElement = MeshBatch.Elements[0];
+		BatchElement.PrimitiveUniformBuffer = GetUniformBuffer();
+		BatchElement.IndexBuffer = &GLidarPointCloudIndexBuffer;
+		BatchElement.FirstIndex = bUsesSprites ? 0 : GLidarPointCloudIndexBuffer.PointOffset;
+		BatchElement.MinVertexIndex = 0;
+		BatchElement.NumPrimitives = Node.NumVisiblePoints * (bUsesSprites ? 2 : 1);
+		BatchElement.MaxVertexIndex = Node.NumVisiblePoints * (bUsesSprites ? 4 : 1);
+		BatchElement.UserData = UserData;
+		BatchElement.VertexFactoryUserData = RenderData.bUseStaticBuffers ? Node.VertexFactory->GetUniformBuffer() : Node.DataCache->GetUniformBuffer();
+	}
+
+#if RHI_RAYTRACING
+	virtual bool IsRayTracingRelevant() const override { return true; }
+
+	virtual void GetDynamicRayTracingInstances(FRayTracingMaterialGatheringContext& Context, TArray<FRayTracingInstance>& OutRayTracingInstances) override
+	{
+		if (RenderData.NumElements)
+		{
+			TArray<FLidarPointCloudBatchElementUserData>& UserData = Context.RayTracingMeshResourceCollector.AllocateOneFrameResource<FLidarOneFrameResource>().Payload;
+			UserData.Reserve(RenderData.SelectedNodes.Num());
+
+			CachedRayTracingMaterials.Reset();
+			const FMatrix& ThisLocalToWorld = GetLocalToWorld();
+
+			for (const FLidarPointCloudProxyUpdateDataNode& Node : RenderData.SelectedNodes)
+			{
+				if (Node.RayTracingGeometry.IsValid() && Node.RayTracingGeometry->IsInitialized())
+				{
+					TArray<FMeshBatch> &NodeRayTracingMaterials = CachedRayTracingMaterials.AddDefaulted_GetRef();
+					FMeshBatch &MeshBatch = NodeRayTracingMaterials.AddDefaulted_GetRef();
+					SetupMeshBatch(MeshBatch, Node, &UserData[UserData.Add(BuildUserDataElement(nullptr, Node))]);
+					MeshBatch.SegmentIndex = 0;
+					MeshBatch.CastRayTracedShadow = IsShadowCast(Context.ReferenceView);
+
+					FLidarPointCloudRayTracingGeometry* Geometry = Node.RayTracingGeometry.Get();
+					FRayTracingInstance &RayTracingInstance = OutRayTracingInstances.AddDefaulted_GetRef();
+					RayTracingInstance.Geometry = Geometry;
+					RayTracingInstance.InstanceTransformsView = MakeArrayView(&ThisLocalToWorld, 1);
+					RayTracingInstance.MaterialsView = MakeArrayView(NodeRayTracingMaterials);
+					RayTracingInstance.BuildInstanceMaskAndFlags(GetScene().GetFeatureLevel());
+
+					Context.DynamicRayTracingGeometriesToUpdate.Add(
+						FRayTracingDynamicGeometryUpdateParams
+						{
+							NodeRayTracingMaterials,
+							false,
+							Geometry->GetNumVertices(),
+							Geometry->GetBufferSize(),
+							Geometry->GetNumPrimitives(),
+							Geometry,
+							nullptr,
+							true
+						}
+					);
+				}
+			}
+		}
+	}
+#endif // RHI_RAYTRACING
+	
 public:
 	TSharedPtr<FLidarPointCloudSceneProxyWrapper, ESPMode::ThreadSafe> ProxyWrapper;
 
@@ -403,6 +465,8 @@ private:
 
 	FLidarPointCloudRenderBuffer* TreeBuffer;
 	FMaterialRelevance MaterialRelevance;
+
+	TArray<TArray<FMeshBatch>> CachedRayTracingMaterials;
 	
 	bool bCompatiblePlatform;
 	
diff --git a/Engine/Plugins/Enterprise/LidarPointCloud/Source/LidarPointCloudRuntime/Public/LidarPointCloudOctree.h b/Engine/Plugins/Enterprise/LidarPointCloud/Source/LidarPointCloudRuntime/Public/LidarPointCloudOctree.h
--- a/Engine/Plugins/Enterprise/LidarPointCloud/Source/LidarPointCloudRuntime/Public/LidarPointCloudOctree.h
+++ b/Engine/Plugins/Enterprise/LidarPointCloud/Source/LidarPointCloudRuntime/Public/LidarPointCloudOctree.h
@@ -19,6 +19,7 @@ struct FLidarPointCloudTraversalOctree;
 struct FLidarPointCloudTraversalOctreeNode;
 class FLidarPointCloudRenderBuffer;
 class FLidarPointCloudVertexFactory;
+class FLidarPointCloudRayTracingGeometry;
 
 namespace LidarPointCloudMeshing
 {
@@ -104,7 +105,8 @@ private:
 	/** Holds render data for this node */
 	TSharedPtr<FLidarPointCloudRenderBuffer> DataCache;
 	TSharedPtr<FLidarPointCloudVertexFactory> VertexFactory;
-
+	TSharedPtr<FLidarPointCloudRayTracingGeometry> RayTracingGeometry;
+	
 	bool bRenderDataDirty;
 
 	/** Used to keep track, which data is available for rendering */
@@ -133,12 +135,15 @@ public:
 
 	/** Return a pointer to the vertex factory containing pre-cached geometry */
 	FORCEINLINE TSharedPtr<FLidarPointCloudVertexFactory> GetVertexFactory() { return VertexFactory; }
+	
+	/** Return a pointer to the ray tracing geometry */
+	FORCEINLINE TSharedPtr<FLidarPointCloudRayTracingGeometry> GetRayTracingGeometry() { return RayTracingGeometry; }
 
 	/**
 	 * Builds and updates the necessary render data buffers
 	 * Returns true if successful
 	 */
-	bool BuildDataCache(bool bUseStaticBuffers);
+	bool BuildDataCache(bool bUseStaticBuffers, bool bUseRayTracing);
 
 	/** Returns the sum of grid and padding points allocated to this node. */
 	FORCEINLINE uint32 GetNumPoints() const { return NumPoints; }
diff --git a/Engine/Plugins/Enterprise/LidarPointCloud/Source/LidarPointCloudRuntime/Public/LidarPointCloudSettings.h b/Engine/Plugins/Enterprise/LidarPointCloud/Source/LidarPointCloudRuntime/Public/LidarPointCloudSettings.h
--- a/Engine/Plugins/Enterprise/LidarPointCloud/Source/LidarPointCloudRuntime/Public/LidarPointCloudSettings.h
+++ b/Engine/Plugins/Enterprise/LidarPointCloud/Source/LidarPointCloudRuntime/Public/LidarPointCloudSettings.h
@@ -94,6 +94,13 @@ public:
 	UPROPERTY(config, EditAnywhere, Category=Performance)
 	bool bUseFastRendering;
 
+	/**
+	 * If enabled, Lidar assets will become visible to Ray Tracing.
+	 * Warning, this will significantly increase VRAM usage!
+	 */
+	UPROPERTY(config, EditAnywhere, Category=RayTracing)
+	bool bEnableLidarRayTracing;
+	
 	/** Affects the size of per-thread data for the meshing algorithm. */
 	UPROPERTY(config, EditAnywhere, Category=Collision)
 	int32 MeshingBatchSize;
diff --git a/Engine/Plugins/Enterprise/LidarPointCloud/Source/LidarPointCloudRuntime/Public/Rendering/LidarPointCloudRendering.h b/Engine/Plugins/Enterprise/LidarPointCloud/Source/LidarPointCloudRuntime/Public/Rendering/LidarPointCloudRendering.h
--- a/Engine/Plugins/Enterprise/LidarPointCloud/Source/LidarPointCloudRuntime/Public/Rendering/LidarPointCloudRendering.h
+++ b/Engine/Plugins/Enterprise/LidarPointCloud/Source/LidarPointCloudRuntime/Public/Rendering/LidarPointCloudRendering.h
@@ -29,6 +29,7 @@ public:
 	 */
 	TSharedPtr<class FLidarPointCloudRenderBuffer> DataCache;
 	TSharedPtr<class FLidarPointCloudVertexFactory> VertexFactory;
+	TSharedPtr<class FLidarPointCloudRayTracingGeometry> RayTracingGeometry;
 
 	FLidarPointCloudProxyUpdateDataNode() : FLidarPointCloudProxyUpdateDataNode(0, 0, nullptr) {}
 	FLidarPointCloudProxyUpdateDataNode(uint8 VirtualDepth, int64 NumVisiblePoints, FLidarPointCloudOctreeNode* DataNode)
@@ -37,11 +38,12 @@ public:
 		, NumVisiblePoints(NumVisiblePoints)
 		, DataCache(nullptr)
 		, VertexFactory(nullptr)
+		, RayTracingGeometry(nullptr)
 	{
 	}
 
 	/** Passthrough method */
-	bool BuildDataCache(bool bUseStaticBuffers);
+	bool BuildDataCache(bool bUseStaticBuffers, bool bUseRayTracing);
 };
 
 /** Used to pass data to RT to update the proxy's render data */
@@ -58,6 +60,7 @@ struct FLidarPointCloudProxyUpdateData
 	float RootCellSize;
 	
 	bool bUseStaticBuffers;
+	bool bUseRayTracing;
 
 	TArray<uint32> TreeStructure;
 
