Lidar - Selection tools not working

diff --git a/Engine/Plugins/Enterprise/LidarPointCloud/Source/LidarPointCloudRuntime/Private/LidarPointCloudOctree.cpp b/Engine/Plugins/Enterprise/LidarPointCloud/Source/LidarPointCloudRuntime/Private/LidarPointCloudOctree.cpp
--- a/Engine/Plugins/Enterprise/LidarPointCloud/Source/LidarPointCloudRuntime/Private/LidarPointCloudOctree.cpp
+++ b/Engine/Plugins/Enterprise/LidarPointCloud/Source/LidarPointCloudRuntime/Private/LidarPointCloudOctree.cpp
@@ -1,5 +1,5 @@
 // Copyright Epic Games, Inc. All Rights Reserved.
- 
+
 #include "LidarPointCloudOctree.h"
 #include "LidarPointCloudOctreeMacros.h"
 #include "LidarPointCloud.h"
@@ -256,17 +256,17 @@ bool FLidarPointCloudOctreeNode::BuildDataCache(bool bUseStaticBuffers, bool bUs
 
 		if (bRenderDataDirty)
 		{
-			if (DataCache.IsValid() && !DataCache->IsInitialized())
+			if (DataCache.IsValid())
 			{
 				DataCache->Initialize(GetData(), GetNumVisiblePoints());
 			}
 
-			if (VertexFactory.IsValid() && !VertexFactory->IsInitialized())
+			if (VertexFactory.IsValid())
 			{
 				VertexFactory->Initialize(GetData(), GetNumVisiblePoints());
 			}
 
-			if(RayTracingGeometry.IsValid() && !RayTracingGeometry->IsInitialized())
+			if(RayTracingGeometry.IsValid())
 			{
 				RayTracingGeometry->Initialize(GetNumVisiblePoints());
 			}
diff --git a/Engine/Plugins/Enterprise/LidarPointCloud/Source/LidarPointCloudRuntime/Private/Rendering/LidarPointCloudRenderBuffers.cpp b/Engine/Plugins/Enterprise/LidarPointCloud/Source/LidarPointCloudRuntime/Private/Rendering/LidarPointCloudRenderBuffers.cpp
--- a/Engine/Plugins/Enterprise/LidarPointCloud/Source/LidarPointCloudRuntime/Private/Rendering/LidarPointCloudRenderBuffers.cpp
+++ b/Engine/Plugins/Enterprise/LidarPointCloud/Source/LidarPointCloudRuntime/Private/Rendering/LidarPointCloudRenderBuffers.cpp
@@ -136,6 +136,11 @@ void FLidarPointCloudRenderBuffer::Initialize(FLidarPointCloudPoint* Data, int32
 void FLidarPointCloudRayTracingGeometry::Initialize(int32 NumPoints)
 {
 #if RHI_RAYTRACING
+	if (IsInitialized())
+	{
+		ReleaseResource();
+	}
+	
 	NumPrimitives = NumPoints * 2;
 	NumVertices = NumPoints * 4;
 	
