Lidar Remove the unrecognized pane from the old editor

diff --git a/Engine/Plugins/Enterprise/LidarPointCloud/Source/LidarPointCloudEditor/Private/LidarPointCloudEditor.cpp b/Engine/Plugins/Enterprise/LidarPointCloud/Source/LidarPointCloudEditor/Private/LidarPointCloudEditor.cpp
--- a/Engine/Plugins/Enterprise/LidarPointCloud/Source/LidarPointCloudEditor/Private/LidarPointCloudEditor.cpp
+++ b/Engine/Plugins/Enterprise/LidarPointCloud/Source/LidarPointCloudEditor/Private/LidarPointCloudEditor.cpp
@@ -112,7 +112,7 @@ void FLidarPointCloudEditor::InitPointCloudEditor(const EToolkitMode::Type Mode,
 	Viewport = SNew(SLidarPointCloudEditorViewport).PointCloudEditor(SharedThis(this)).ObjectToEdit(PointCloudBeingEdited);
 
 	// Default layout
-	const TSharedRef<FTabManager::FLayout> StandaloneDefaultLayout = FTabManager::NewLayout("Standalone_LidarPointCloudEditor_Layout_v3")
+	const TSharedRef<FTabManager::FLayout> StandaloneDefaultLayout = FTabManager::NewLayout("Standalone_LidarPointCloudEditor_Layout_v4")
 		->AddArea
 		(
 			FTabManager::NewPrimaryArea()
