Lidar: Possible to create a static mesh with an empty selection
Actions enabled even if no selection is made

diff --git a/Engine/Plugins/Enterprise/LidarPointCloud/Source/LidarPointCloudEditor/Private/LidarPointCloudEdModeToolkit.cpp b/Engine/Plugins/Enterprise/LidarPointCloud/Source/LidarPointCloudEditor/Private/LidarPointCloudEdModeToolkit.cpp
--- a/Engine/Plugins/Enterprise/LidarPointCloud/Source/LidarPointCloudEditor/Private/LidarPointCloudEdModeToolkit.cpp
+++ b/Engine/Plugins/Enterprise/LidarPointCloud/Source/LidarPointCloudEditor/Private/LidarPointCloudEdModeToolkit.cpp
@@ -34,6 +34,7 @@ void SLidarEditorWidget::Construct(const FArguments& InArgs)
 #define HEADEREX(ContText, Category) +SVerticalBox::Slot().AutoHeight().Padding(FMargin(0.0f, 5.0f))[SNew(SHeaderRow)+SHeaderRow::Column(ContText).HAlignCell(HAlign_Left).FillWidth(1).HeaderContentPadding(HeaderPadding)[SNew(STextBlock).Text(LOCTEXT(Category, ContText)).Font(HeaderFont)]]
 #define HEADER(ContText) HEADEREX(ContText, ContText"Header")
 #define BUTTON(Category, ConText, Action) +SHorizontalBox::Slot().Padding(StandardPadding).FillWidth(0.5f)[SNew(SButton).HAlign(HAlign_Center).Text(LOCTEXT(Category,ConText)).OnClicked_Lambda([this]{ Action return FReply::Handled(); })]
+#define BUTTON_ANY(Category, ConText, Action) +SHorizontalBox::Slot().Padding(StandardPadding).FillWidth(0.5f)[SNew(SButton).HAlign(HAlign_Center).Text(LOCTEXT(Category,ConText)).OnClicked_Lambda([this]{ Action return FReply::Handled(); }).IsEnabled(this, &SLidarEditorWidget::IsAnySelection)]
 #define BUTTON_POINTS(Category, ConText, Action) +SHorizontalBox::Slot().Padding(StandardPadding).FillWidth(0.5f)[SNew(SButton).HAlign(HAlign_Center).Text(LOCTEXT(Category,ConText)).OnClicked_Lambda([this]{ Action return FReply::Handled(); }).IsEnabled(this, &SLidarEditorWidget::IsPointSelection)]
 #define BUTTON_ACTORS(Category, ConText, Action) +SHorizontalBox::Slot().Padding(StandardPadding).FillWidth(0.5f)[SNew(SButton).HAlign(HAlign_Center).Text(LOCTEXT(Category,ConText)).OnClicked_Lambda([this]{ Action return FReply::Handled(); }).IsEnabled(this, &SLidarEditorWidget::IsActorSelection)]
 
@@ -44,7 +45,7 @@ void SLidarEditorWidget::Construct(const FArguments& InArgs)
 		.AutoHeight()
 		[
 			SNew(SHorizontalBox)
-			BUTTON("SelectionClear", "Clear",
+			BUTTON_ANY("SelectionClear", "Clear",
 			{
 				if(IsActorSelection())
 				{
@@ -249,7 +250,7 @@ void SLidarEditorWidget::Construct(const FArguments& InArgs)
 		.AutoHeight()
 		[
 			SNew(SHorizontalBox)
-			BUTTON("NormalsCalculate", "Calculate Normals",
+			BUTTON_ANY("NormalsCalculate", "Calculate Normals",
 			{
 				FLidarPointCloudEditorHelper::SetNormalsQuality(NormalsQuality, NormalsNoiseTolerance);
 				if(IsActorSelection())
@@ -337,7 +338,7 @@ void SLidarEditorWidget::Construct(const FArguments& InArgs)
 		.AutoHeight()
 		[
 			SNew(SHorizontalBox)
-			BUTTON("MeshingBuild", "Create Static Mesh", {FLidarPointCloudEditorHelper::MeshSelected(IsPointSelection(), MaxMeshingError, bMergeMeshes, !bMergeMeshes && bRetainTransform);})
+			BUTTON_ANY("MeshingBuild", "Create Static Mesh", {FLidarPointCloudEditorHelper::MeshSelected(IsPointSelection(), MaxMeshingError, bMergeMeshes, !bMergeMeshes && bRetainTransform);})
 		]
 		HEADER("Alignment")
 		+ SVerticalBox::Slot()
@@ -384,6 +385,16 @@ void SLidarEditorWidget::Construct(const FArguments& InArgs)
 }
 END_SLATE_FUNCTION_BUILD_OPTIMIZATION
 
+bool SLidarEditorWidget::IsActorSelection() const
+{
+	return bActorSelection && FLidarPointCloudEditorHelper::AreLidarActorsSelected();
+}
+
+bool SLidarEditorWidget::IsPointSelection() const
+{
+	return !bActorSelection && FLidarPointCloudEditorHelper::AreLidarPointsSelected();
+}
+
 FLidarPointCloudEdModeToolkit::~FLidarPointCloudEdModeToolkit()
 {
 	GetScriptableEditorMode()->GetInteractiveToolsContext(EToolsContextScope::EdMode)->OnToolNotificationMessage.RemoveAll(this);
diff --git a/Engine/Plugins/Enterprise/LidarPointCloud/Source/LidarPointCloudEditor/Private/LidarPointCloudEdModeToolkit.h b/Engine/Plugins/Enterprise/LidarPointCloud/Source/LidarPointCloudEditor/Private/LidarPointCloudEdModeToolkit.h
--- a/Engine/Plugins/Enterprise/LidarPointCloud/Source/LidarPointCloudEditor/Private/LidarPointCloudEdModeToolkit.h
+++ b/Engine/Plugins/Enterprise/LidarPointCloud/Source/LidarPointCloudEditor/Private/LidarPointCloudEdModeToolkit.h
@@ -22,12 +22,14 @@ public:
 public:
 	void Construct(const FArguments& InArgs);
 
-	bool IsActorSelection() const { return bActorSelection; }
-	bool IsPointSelection() const { return !bActorSelection; }
+	bool IsActorSelection() const;
+	bool IsPointSelection() const;
+	bool IsAnySelection() const { return IsActorSelection() || IsPointSelection(); }
 	bool IsBrushVisible() const { return BrushTool != nullptr; }
 
-	EVisibility GetActorVisibility() const { return bActorSelection ? EVisibility::Visible : EVisibility::Collapsed; }
-	EVisibility GetPointVisibility() const { return bActorSelection ? EVisibility::Collapsed : EVisibility::Visible; }
+	EVisibility GetActorVisibility() const { return IsActorSelection() ? EVisibility::Visible : EVisibility::Collapsed; }
+	EVisibility GetPointVisibility() const { return IsPointSelection() ? EVisibility::Visible : EVisibility::Collapsed; }
+	EVisibility GetAnyVisibility() const { return IsAnySelection() ? EVisibility::Visible : EVisibility::Collapsed; }
 	EVisibility GetBrushVisibility() const { return IsBrushVisible() ? EVisibility::Visible : EVisibility::Collapsed; }
 
 private:	
diff --git a/Engine/Plugins/Enterprise/LidarPointCloud/Source/LidarPointCloudEditor/Private/LidarPointCloudEditorHelper.cpp b/Engine/Plugins/Enterprise/LidarPointCloud/Source/LidarPointCloudEditor/Private/LidarPointCloudEditorHelper.cpp
--- a/Engine/Plugins/Enterprise/LidarPointCloud/Source/LidarPointCloudEditor/Private/LidarPointCloudEditorHelper.cpp
+++ b/Engine/Plugins/Enterprise/LidarPointCloud/Source/LidarPointCloudEditor/Private/LidarPointCloudEditorHelper.cpp
@@ -808,6 +808,38 @@ bool FLidarPointCloudEditorHelper::IsPolygonSelfIntersecting(const TArray<FVecto
 	return false;
 }
 
+bool FLidarPointCloudEditorHelper::AreLidarActorsSelected()
+{
+	for (FSelectionIterator It(GEditor->GetSelectedActorIterator()); It; ++It)
+	{
+		if(Cast<ALidarPointCloudActor>(*It))
+		{
+			return true;
+		}
+	}
+
+	return false;
+}
+
+bool FLidarPointCloudEditorHelper::AreLidarPointsSelected()
+{
+	for (TObjectIterator<ALidarPointCloudActor> It; It; ++It)
+	{
+		if(const ALidarPointCloudActor* Actor = Cast<ALidarPointCloudActor>(*It))
+		{
+			if(const ULidarPointCloud* PointCloud = Actor->GetPointCloud())
+			{
+				if(PointCloud->HasSelectedPoints())
+				{
+					return true;
+				}
+			}
+		}
+	}
+
+	return false;
+}
+
 void FLidarPointCloudEditorHelper::SelectPointsByConvexVolume(const FConvexVolume& ConvexVolume, ELidarPointCloudSelectionMode SelectionMode)
 {
 	ProcessAll([ConvexVolume, SelectionMode](ALidarPointCloudActor* Actor)
diff --git a/Engine/Plugins/Enterprise/LidarPointCloud/Source/LidarPointCloudEditor/Private/LidarPointCloudEditorHelper.h b/Engine/Plugins/Enterprise/LidarPointCloud/Source/LidarPointCloudEditor/Private/LidarPointCloudEditorHelper.h
--- a/Engine/Plugins/Enterprise/LidarPointCloud/Source/LidarPointCloudEditor/Private/LidarPointCloudEditorHelper.h
+++ b/Engine/Plugins/Enterprise/LidarPointCloud/Source/LidarPointCloudEditor/Private/LidarPointCloudEditorHelper.h
@@ -62,4 +62,7 @@ public:
 	static bool RayTracePointClouds(const FLidarPointCloudRay& Ray, float RadiusMulti, FVector3f& OutHitLocation);
 
 	static bool IsPolygonSelfIntersecting(const TArray<FVector2D>& Points, bool bAllowLooping);
+
+	static bool AreLidarActorsSelected();
+	static bool AreLidarPointsSelected();
 };
diff --git a/Engine/Plugins/Enterprise/LidarPointCloud/Source/LidarPointCloudRuntime/Private/LidarPointCloud.cpp b/Engine/Plugins/Enterprise/LidarPointCloud/Source/LidarPointCloudRuntime/Private/LidarPointCloud.cpp
--- a/Engine/Plugins/Enterprise/LidarPointCloud/Source/LidarPointCloudRuntime/Private/LidarPointCloud.cpp
+++ b/Engine/Plugins/Enterprise/LidarPointCloud/Source/LidarPointCloudRuntime/Private/LidarPointCloud.cpp
@@ -614,17 +614,17 @@ void ULidarPointCloud::InvertSelection()
 	Octree.InvertSelection();
 }
 
-int64 ULidarPointCloud::NumSelectedPoints()
+int64 ULidarPointCloud::NumSelectedPoints() const
 {
 	return Octree.NumSelectedPoints();
 }
 
-bool ULidarPointCloud::HasSelectedPoints()
+bool ULidarPointCloud::HasSelectedPoints() const
 {
 	return Octree.HasSelectedPoints();
 }
 
-void ULidarPointCloud::GetSelectedPointsAsCopies(TArray64<FLidarPointCloudPoint>& SelectedPoints, FTransform Transform)
+void ULidarPointCloud::GetSelectedPointsAsCopies(TArray64<FLidarPointCloudPoint>& SelectedPoints, FTransform Transform) const
 {
 	Transform.AddToTranslation(LocationOffset);
 	Octree.GetSelectedPointsAsCopies(SelectedPoints, Transform);
diff --git a/Engine/Plugins/Enterprise/LidarPointCloud/Source/LidarPointCloudRuntime/Private/LidarPointCloudOctree.cpp b/Engine/Plugins/Enterprise/LidarPointCloud/Source/LidarPointCloudRuntime/Private/LidarPointCloudOctree.cpp
--- a/Engine/Plugins/Enterprise/LidarPointCloud/Source/LidarPointCloudRuntime/Private/LidarPointCloudOctree.cpp
+++ b/Engine/Plugins/Enterprise/LidarPointCloud/Source/LidarPointCloudRuntime/Private/LidarPointCloudOctree.cpp
@@ -1915,20 +1915,20 @@ void FLidarPointCloudOctree::InvertSelection()
 	}, true);
 }
 
-int64 FLidarPointCloudOctree::NumSelectedPoints()
+int64 FLidarPointCloudOctree::NumSelectedPoints() const
 {
 	int64 Count = 0;
 	ITERATE_SELECTED({ ++Count; }, {});
 	return Count;
 }
 
-bool FLidarPointCloudOctree::HasSelectedPoints()
+bool FLidarPointCloudOctree::HasSelectedPoints() const
 {
 	ITERATE_SELECTED({ return true; }, {});
 	return false;
 }
 
-void FLidarPointCloudOctree::GetSelectedPointsAsCopies(TArray64<FLidarPointCloudPoint>& SelectedPoints, const FTransform& Transform)
+void FLidarPointCloudOctree::GetSelectedPointsAsCopies(TArray64<FLidarPointCloudPoint>& SelectedPoints, const FTransform& Transform) const
 {
 	const FTransform3f Transform3F = (FTransform3f)Transform;
 
diff --git a/Engine/Plugins/Enterprise/LidarPointCloud/Source/LidarPointCloudRuntime/Public/LidarPointCloud.h b/Engine/Plugins/Enterprise/LidarPointCloud/Source/LidarPointCloudRuntime/Public/LidarPointCloud.h
--- a/Engine/Plugins/Enterprise/LidarPointCloud/Source/LidarPointCloudRuntime/Public/LidarPointCloud.h
+++ b/Engine/Plugins/Enterprise/LidarPointCloud/Source/LidarPointCloudRuntime/Public/LidarPointCloud.h
@@ -425,9 +425,9 @@ public:
 	void HideSelected();
 	void DeleteSelected();
 	void InvertSelection();
-	int64 NumSelectedPoints();
-	bool HasSelectedPoints();
-	void GetSelectedPointsAsCopies(TArray64<FLidarPointCloudPoint>& SelectedPoints, FTransform Transform);
+	int64 NumSelectedPoints() const;
+	bool HasSelectedPoints() const;
+	void GetSelectedPointsAsCopies(TArray64<FLidarPointCloudPoint>& SelectedPoints, FTransform Transform) const;
 	void CalculateNormalsForSelection();
 	void ClearSelection();
 	void BuildStaticMeshBuffersForSelection(float CellSize, LidarPointCloudMeshing::FMeshBuffers* OutMeshBuffers, const FTransform& Transform);
diff --git a/Engine/Plugins/Enterprise/LidarPointCloud/Source/LidarPointCloudRuntime/Public/LidarPointCloudOctree.h b/Engine/Plugins/Enterprise/LidarPointCloud/Source/LidarPointCloudRuntime/Public/LidarPointCloudOctree.h
--- a/Engine/Plugins/Enterprise/LidarPointCloud/Source/LidarPointCloudRuntime/Public/LidarPointCloudOctree.h
+++ b/Engine/Plugins/Enterprise/LidarPointCloud/Source/LidarPointCloudRuntime/Public/LidarPointCloudOctree.h
@@ -488,9 +488,9 @@ public:
 	void HideSelected();
 	void DeleteSelected();
 	void InvertSelection();
-	int64 NumSelectedPoints();
-	bool HasSelectedPoints();
-	void GetSelectedPointsAsCopies(TArray64<FLidarPointCloudPoint>& SelectedPoints, const FTransform& Transform);
+	int64 NumSelectedPoints() const;
+	bool HasSelectedPoints() const;
+	void GetSelectedPointsAsCopies(TArray64<FLidarPointCloudPoint>& SelectedPoints, const FTransform& Transform) const;
 	void GetSelectedPointsInBox(TArray64<const FLidarPointCloudPoint*>& SelectedPoints, const FBox& Box) const;
 	void CalculateNormalsForSelection(FThreadSafeBool* bCancelled, int32 Quality, float Tolerance);
 	void ClearSelection();
