Fixed Normals are not calculated on point selections

diff --git a/Engine/Plugins/Enterprise/LidarPointCloud/Source/LidarPointCloudEditor/Private/LidarPointCloudEdModeToolkit.cpp b/Engine/Plugins/Enterprise/LidarPointCloud/Source/LidarPointCloudEditor/Private/LidarPointCloudEdModeToolkit.cpp
--- a/Engine/Plugins/Enterprise/LidarPointCloud/Source/LidarPointCloudEditor/Private/LidarPointCloudEdModeToolkit.cpp
+++ b/Engine/Plugins/Enterprise/LidarPointCloud/Source/LidarPointCloudEditor/Private/LidarPointCloudEdModeToolkit.cpp
@@ -253,13 +253,13 @@ void SLidarEditorWidget::Construct(const FArguments& InArgs)
 			BUTTON_ANY("NormalsCalculate", "Calculate Normals",
 			{
 				FLidarPointCloudEditorHelper::SetNormalsQuality(NormalsQuality, NormalsNoiseTolerance);
-				if(IsActorSelection())
+				if(IsPointSelection())
 				{
-					FLidarPointCloudEditorHelper::CalculateNormals();
+					FLidarPointCloudEditorHelper::CalculateNormalsForSelection();
 				}
 				else
 				{
-					FLidarPointCloudEditorHelper::CalculateNormalsForSelection();
+					FLidarPointCloudEditorHelper::CalculateNormals();
 				}
 			})
 		]
@@ -387,7 +387,7 @@ END_SLATE_FUNCTION_BUILD_OPTIMIZATION
 
 bool SLidarEditorWidget::IsActorSelection() const
 {
-	return bActorSelection && FLidarPointCloudEditorHelper::AreLidarActorsSelected();
+	return bActorSelection && FLidarPointCloudEditorHelper::AreLidarActorsSelected() && !FLidarPointCloudEditorHelper::AreLidarPointsSelected();
 }
 
 bool SLidarEditorWidget::IsPointSelection() const
diff --git a/Engine/Plugins/Enterprise/LidarPointCloud/Source/LidarPointCloudEditor/Private/LidarPointCloudEditorHelper.cpp b/Engine/Plugins/Enterprise/LidarPointCloud/Source/LidarPointCloudEditor/Private/LidarPointCloudEditorHelper.cpp
--- a/Engine/Plugins/Enterprise/LidarPointCloud/Source/LidarPointCloudEditor/Private/LidarPointCloudEditorHelper.cpp
+++ b/Engine/Plugins/Enterprise/LidarPointCloud/Source/LidarPointCloudEditor/Private/LidarPointCloudEditorHelper.cpp
@@ -594,7 +594,7 @@ void FLidarPointCloudEditorHelper::CalculateNormalsForSelection()
 	{
 		if(ULidarPointCloud* PointCloud = Actor->GetPointCloud())
 		{
-			PointCloud->CalculateNormals(nullptr, nullptr);
+			PointCloud->CalculateNormalsForSelection();
 		}
 	});
 }
@@ -663,7 +663,7 @@ void FLidarPointCloudEditorHelper::CalculateNormals()
 	{
 		if(ULidarPointCloud* PointCloud = Actor->GetPointCloud())
 		{
-			PointCloud->CalculateNormalsForSelection();
+			PointCloud->CalculateNormals(nullptr, nullptr);
 		}
 	});
 }
