Fixed Orient facing normal refresh glitch

diff --git a/Engine/Plugins/Enterprise/LidarPointCloud/Source/LidarPointCloudRuntime/Private/LidarPointCloud.cpp b/Engine/Plugins/Enterprise/LidarPointCloud/Source/LidarPointCloudRuntime/Private/LidarPointCloud.cpp
--- a/Engine/Plugins/Enterprise/LidarPointCloud/Source/LidarPointCloudRuntime/Private/LidarPointCloud.cpp
+++ b/Engine/Plugins/Enterprise/LidarPointCloud/Source/LidarPointCloudRuntime/Private/LidarPointCloud.cpp
@@ -632,7 +632,12 @@ void ULidarPointCloud::GetSelectedPointsAsCopies(TArray64<FLidarPointCloudPoint>
 
 void ULidarPointCloud::CalculateNormalsForSelection()
 {
-	Octree.CalculateNormalsForSelection(nullptr, NormalsQuality, NormalsNoiseTolerance);
+	TArray64<FLidarPointCloudPoint*>* SelectedPoints = new TArray64<FLidarPointCloudPoint*>();
+	Octree.GetSelectedPoints(*SelectedPoints);
+	CalculateNormals(SelectedPoints, [SelectedPoints]
+	{
+		delete SelectedPoints;
+	});
 }
 
 void ULidarPointCloud::ClearSelection()
@@ -1348,6 +1353,7 @@ void ULidarPointCloud::CalculateNormals(TArray64<FLidarPointCloudPoint*>* Points
 		{
 			MarkPackageDirty();
 			Notification->Close(!bAsyncCancelled);
+			OnPointCloudNormalsUpdatedEvent.Broadcast();
 		});
 
 		if (_CompletionCallback)
diff --git a/Engine/Plugins/Enterprise/LidarPointCloud/Source/LidarPointCloudRuntime/Private/LidarPointCloudComponent.cpp b/Engine/Plugins/Enterprise/LidarPointCloud/Source/LidarPointCloudRuntime/Private/LidarPointCloudComponent.cpp
--- a/Engine/Plugins/Enterprise/LidarPointCloud/Source/LidarPointCloudRuntime/Private/LidarPointCloudComponent.cpp
+++ b/Engine/Plugins/Enterprise/LidarPointCloud/Source/LidarPointCloudRuntime/Private/LidarPointCloudComponent.cpp
@@ -82,6 +82,7 @@ void ULidarPointCloudComponent::AttachPointCloudListener()
 	{
 		PointCloud->OnPointCloudRebuilt().AddUObject(this, &ULidarPointCloudComponent::OnPointCloudRebuilt);
 		PointCloud->OnPointCloudCollisionUpdated().AddUObject(this, &ULidarPointCloudComponent::OnPointCloudCollisionUpdated);
+		PointCloud->OnPointCloudNormalsUpdated().AddUObject(this, &ULidarPointCloudComponent::OnPointCloudNormalsUpdated);
 	}
 }
 
@@ -91,6 +92,7 @@ void ULidarPointCloudComponent::RemovePointCloudListener()
 	{
 		PointCloud->OnPointCloudRebuilt().RemoveAll(this);
 		PointCloud->OnPointCloudCollisionUpdated().RemoveAll(this);
+		PointCloud->OnPointCloudNormalsUpdated().RemoveAll(this);
 	}
 }
 
@@ -122,6 +124,12 @@ void ULidarPointCloudComponent::OnPointCloudCollisionUpdated()
 	MarkRenderStateDirty();
 }
 
+void ULidarPointCloudComponent::OnPointCloudNormalsUpdated()
+{
+	PointOrientation = ELidarPointCloudSpriteOrientation::PreferFacingNormal;
+	MarkRenderStateDirty();
+}
+
 void ULidarPointCloudComponent::PostPointCloudSet()
 {
 	AttachPointCloudListener();
diff --git a/Engine/Plugins/Enterprise/LidarPointCloud/Source/LidarPointCloudRuntime/Private/LidarPointCloudOctree.cpp b/Engine/Plugins/Enterprise/LidarPointCloud/Source/LidarPointCloudRuntime/Private/LidarPointCloudOctree.cpp
--- a/Engine/Plugins/Enterprise/LidarPointCloud/Source/LidarPointCloudRuntime/Private/LidarPointCloudOctree.cpp
+++ b/Engine/Plugins/Enterprise/LidarPointCloud/Source/LidarPointCloudRuntime/Private/LidarPointCloudOctree.cpp
@@ -1928,6 +1928,11 @@ bool FLidarPointCloudOctree::HasSelectedPoints() const
 	return false;
 }
 
+void FLidarPointCloudOctree::GetSelectedPoints(TArray64<FLidarPointCloudPoint*>& SelectedPoints) const
+{
+	ITERATE_SELECTED({ SelectedPoints.Add(Point); }, {});
+}
+
 void FLidarPointCloudOctree::GetSelectedPointsAsCopies(TArray64<FLidarPointCloudPoint>& SelectedPoints, const FTransform& Transform) const
 {
 	const FTransform3f Transform3F = (FTransform3f)Transform;
@@ -1951,21 +1956,6 @@ void FLidarPointCloudOctree::GetSelectedPointsInBox(TArray64<const FLidarPointCl
 	});
 }
 
-void FLidarPointCloudOctree::CalculateNormalsForSelection(FThreadSafeBool* bCancelled, int32 Quality, float Tolerance)
-{
-	TArray64<FLidarPointCloudPoint*> SelectedPoints;
-
-	ITERATE_SELECTED({
-		SelectedPoints.Add(Point);
-		Point->Normal.Reset();
-	},{
-		CurrentNode->bCanReleaseData = false;
-		CurrentNode->bRenderDataDirty = true;
-	});
-	
-	LidarPointCloudMeshing::CalculateNormals(this, bCancelled, Quality, Tolerance, SelectedPoints);
-}
-
 void FLidarPointCloudOctree::ClearSelection()
 {
 	ITERATE_SELECTED({
diff --git a/Engine/Plugins/Enterprise/LidarPointCloud/Source/LidarPointCloudRuntime/Public/LidarPointCloud.h b/Engine/Plugins/Enterprise/LidarPointCloud/Source/LidarPointCloudRuntime/Public/LidarPointCloud.h
--- a/Engine/Plugins/Enterprise/LidarPointCloud/Source/LidarPointCloudRuntime/Public/LidarPointCloud.h
+++ b/Engine/Plugins/Enterprise/LidarPointCloud/Source/LidarPointCloudRuntime/Public/LidarPointCloud.h
@@ -165,6 +165,7 @@ private:
 
 	FOnPointCloudChanged OnPointCloudRebuiltEvent;
 	FOnPointCloudChanged OnPointCloudUpdateCollisionEvent;
+	FOnPointCloudChanged OnPointCloudNormalsUpdatedEvent;
 	FOnPointCloudChanged OnPreSaveCleanupEvent;
 
 public:
@@ -172,6 +173,7 @@ public:
 
 	FOnPointCloudChanged& OnPointCloudRebuilt() { return OnPointCloudRebuiltEvent; }
 	FOnPointCloudChanged& OnPointCloudCollisionUpdated() { return OnPointCloudUpdateCollisionEvent; }
+	FOnPointCloudChanged& OnPointCloudNormalsUpdated() { return OnPointCloudNormalsUpdatedEvent; }
 	FOnPointCloudChanged& OnPreSaveCleanup() { return OnPreSaveCleanupEvent; }
 
 	// Begin UObject Interface.
diff --git a/Engine/Plugins/Enterprise/LidarPointCloud/Source/LidarPointCloudRuntime/Public/LidarPointCloudComponent.h b/Engine/Plugins/Enterprise/LidarPointCloud/Source/LidarPointCloudRuntime/Public/LidarPointCloudComponent.h
--- a/Engine/Plugins/Enterprise/LidarPointCloud/Source/LidarPointCloudRuntime/Public/LidarPointCloudComponent.h
+++ b/Engine/Plugins/Enterprise/LidarPointCloud/Source/LidarPointCloudRuntime/Public/LidarPointCloudComponent.h
@@ -583,6 +583,7 @@ private:
 	void RemovePointCloudListener();
 	void OnPointCloudRebuilt();
 	void OnPointCloudCollisionUpdated();
+	void OnPointCloudNormalsUpdated();
 
 	void PostPointCloudSet();
 
diff --git a/Engine/Plugins/Enterprise/LidarPointCloud/Source/LidarPointCloudRuntime/Public/LidarPointCloudOctree.h b/Engine/Plugins/Enterprise/LidarPointCloud/Source/LidarPointCloudRuntime/Public/LidarPointCloudOctree.h
--- a/Engine/Plugins/Enterprise/LidarPointCloud/Source/LidarPointCloudRuntime/Public/LidarPointCloudOctree.h
+++ b/Engine/Plugins/Enterprise/LidarPointCloud/Source/LidarPointCloudRuntime/Public/LidarPointCloudOctree.h
@@ -490,9 +490,9 @@ public:
 	void InvertSelection();
 	int64 NumSelectedPoints() const;
 	bool HasSelectedPoints() const;
+	void GetSelectedPoints(TArray64<FLidarPointCloudPoint*>& SelectedPoints) const;
 	void GetSelectedPointsAsCopies(TArray64<FLidarPointCloudPoint>& SelectedPoints, const FTransform& Transform) const;
 	void GetSelectedPointsInBox(TArray64<const FLidarPointCloudPoint*>& SelectedPoints, const FBox& Box) const;
-	void CalculateNormalsForSelection(FThreadSafeBool* bCancelled, int32 Quality, float Tolerance);
 	void ClearSelection();
 	void BuildStaticMeshBuffersForSelection(float CellSize, LidarPointCloudMeshing::FMeshBuffers* OutMeshBuffers, const FTransform& Transform);
 #endif
