Fixed Green box shaped icon displayed on Lidar assets that don't have any collisions

diff --git a/Engine/Plugins/Enterprise/LidarPointCloud/Source/LidarPointCloudEditor/Private/LidarPointCloudEditorViewport.cpp b/Engine/Plugins/Enterprise/LidarPointCloud/Source/LidarPointCloudEditor/Private/LidarPointCloudEditorViewport.cpp
--- a/Engine/Plugins/Enterprise/LidarPointCloud/Source/LidarPointCloudEditor/Private/LidarPointCloudEditorViewport.cpp
+++ b/Engine/Plugins/Enterprise/LidarPointCloud/Source/LidarPointCloudEditor/Private/LidarPointCloudEditorViewport.cpp
@@ -54,9 +54,6 @@ public:
 
 void SLidarPointCloudEditorViewport::Construct(const FArguments& InArgs)
 {
-	PreviewScene.SetFloorVisibility(false);
-	PreviewScene.SetEnvironmentVisibility(false);
-
 	PointCloudEditorPtr = InArgs._PointCloudEditor;
 
 	PointCloud = InArgs._ObjectToEdit;
@@ -162,7 +159,7 @@ bool SLidarPointCloudEditorViewport::IsVisible() const
 
 TSharedRef<FEditorViewportClient> SLidarPointCloudEditorViewport::MakeEditorViewportClient()
 {
-	EditorViewportClient = MakeShareable(new FLidarPointCloudEditorViewportClient(PointCloudEditorPtr, SharedThis(this), &PreviewScene, PointCloud, PreviewCloudComponent));
+	EditorViewportClient = MakeShareable(new FLidarPointCloudEditorViewportClient(PointCloudEditorPtr, SharedThis(this), &PreviewScene, PreviewCloudComponent));
 
 	EditorViewportClient->bSetListenerPosition = false;
 
diff --git a/Engine/Plugins/Enterprise/LidarPointCloud/Source/LidarPointCloudEditor/Private/LidarPointCloudEditorViewport.h b/Engine/Plugins/Enterprise/LidarPointCloud/Source/LidarPointCloudEditor/Private/LidarPointCloudEditorViewport.h
--- a/Engine/Plugins/Enterprise/LidarPointCloud/Source/LidarPointCloudEditor/Private/LidarPointCloudEditorViewport.h
+++ b/Engine/Plugins/Enterprise/LidarPointCloud/Source/LidarPointCloudEditor/Private/LidarPointCloudEditorViewport.h
@@ -97,7 +97,7 @@ private:
 	TWeakPtr<FLidarPointCloudEditor> PointCloudEditorPtr;
 
 	/** The scene for this viewport. */
-	FAdvancedPreviewScene PreviewScene;
+	FPreviewScene PreviewScene;
 
 	/** Editor viewport client */
 	TSharedPtr<FLidarPointCloudEditorViewportClient> EditorViewportClient;
diff --git a/Engine/Plugins/Enterprise/LidarPointCloud/Source/LidarPointCloudEditor/Private/LidarPointCloudEditorViewportClient.cpp b/Engine/Plugins/Enterprise/LidarPointCloud/Source/LidarPointCloudEditor/Private/LidarPointCloudEditorViewportClient.cpp
--- a/Engine/Plugins/Enterprise/LidarPointCloud/Source/LidarPointCloudEditor/Private/LidarPointCloudEditorViewportClient.cpp
+++ b/Engine/Plugins/Enterprise/LidarPointCloud/Source/LidarPointCloudEditor/Private/LidarPointCloudEditorViewportClient.cpp
@@ -31,7 +31,7 @@ namespace UE {
 	}
 }
 
-FLidarPointCloudEditorViewportClient::FLidarPointCloudEditorViewportClient(TWeakPtr<FLidarPointCloudEditor> InPointCloudEditor, const TSharedRef<SLidarPointCloudEditorViewport>& InPointCloudEditorViewport, FAdvancedPreviewScene* InPreviewScene, ULidarPointCloud* InPreviewPointCloud, ULidarPointCloudComponent* InPreviewPointCloudComponent)
+FLidarPointCloudEditorViewportClient::FLidarPointCloudEditorViewportClient(TWeakPtr<FLidarPointCloudEditor> InPointCloudEditor, const TSharedRef<SLidarPointCloudEditorViewport>& InPointCloudEditorViewport, FPreviewScene* InPreviewScene, ULidarPointCloudComponent* InPreviewPointCloudComponent)
 	: FEditorViewportClient(nullptr, InPreviewScene, StaticCastSharedRef<SEditorViewport>(InPointCloudEditorViewport))
 	, PointCloudComponent(InPreviewPointCloudComponent)
 	, PointCloudEditorPtr(InPointCloudEditor)
@@ -56,11 +56,6 @@ FLidarPointCloudEditorViewportClient::FLidarPointCloudEditorViewportClient(TWeak
 	OverrideNearClipPlane(1.0f);
 	bUsingOrbitCamera = true;
 
-	AdvancedPreviewScene = static_cast<FAdvancedPreviewScene*>(PreviewScene);
-
-	// Register delegate to update the show flags when the post processing is turned on or off
-	UAssetViewerSettings::Get()->OnAssetViewerSettingsChanged().AddRaw(this, &FLidarPointCloudEditorViewportClient::OnAssetViewerSettingsChanged);
-
 	// Set correct flags according to current profile settings
 	SetAdvancedShowFlagsForScene(UAssetViewerSettings::Get()->Profiles[GetMutableDefault<UEditorPerProjectUserSettings>()->AssetViewerProfileIndex].bPostProcessingEnabled);
 }
@@ -116,19 +111,6 @@ void FLidarPointCloudEditorViewportClient::PerspectiveCameraMoved()
 	ToggleOrbitCamera(bUsingOrbitCamera);
 }
 
-void FLidarPointCloudEditorViewportClient::OnAssetViewerSettingsChanged(const FName& InPropertyName)
-{
-	if (InPropertyName == GET_MEMBER_NAME_CHECKED(FPreviewSceneProfile, bPostProcessingEnabled) || InPropertyName == NAME_None)
-	{
-		UAssetViewerSettings* Settings = UAssetViewerSettings::Get();
-		const int32 ProfileIndex = AdvancedPreviewScene->GetCurrentProfileIndex();
-		if (Settings->Profiles.IsValidIndex(ProfileIndex))
-		{
-			SetAdvancedShowFlagsForScene(Settings->Profiles[ProfileIndex].bPostProcessingEnabled);
-		}
-	}
-}
-
 void FLidarPointCloudEditorViewportClient::SetAdvancedShowFlagsForScene(const bool bAdvancedShowFlags)
 {
 	if (bAdvancedShowFlags)
diff --git a/Engine/Plugins/Enterprise/LidarPointCloud/Source/LidarPointCloudEditor/Private/LidarPointCloudEditorViewportClient.h b/Engine/Plugins/Enterprise/LidarPointCloud/Source/LidarPointCloudEditor/Private/LidarPointCloudEditorViewportClient.h
--- a/Engine/Plugins/Enterprise/LidarPointCloud/Source/LidarPointCloudEditor/Private/LidarPointCloudEditorViewportClient.h
+++ b/Engine/Plugins/Enterprise/LidarPointCloud/Source/LidarPointCloudEditor/Private/LidarPointCloudEditorViewportClient.h
@@ -20,7 +20,7 @@ class ULidarPointCloudComponent;
 class FLidarPointCloudEditorViewportClient : public FEditorViewportClient, public TSharedFromThis<FLidarPointCloudEditorViewportClient>
 {
 public:
-	FLidarPointCloudEditorViewportClient(TWeakPtr<FLidarPointCloudEditor> InPointCloudEditor, const TSharedRef<SLidarPointCloudEditorViewport>& InPointCloudEditorViewport, FAdvancedPreviewScene* InPreviewScene, ULidarPointCloud* InPreviewPointCloud, ULidarPointCloudComponent* InPreviewPointCloudComponent);
+	FLidarPointCloudEditorViewportClient(TWeakPtr<FLidarPointCloudEditor> InPointCloudEditor, const TSharedRef<SLidarPointCloudEditorViewport>& InPointCloudEditorViewport, FPreviewScene* InPreviewScene, ULidarPointCloudComponent* InPreviewPointCloudComponent);
 	~FLidarPointCloudEditorViewportClient();
 
 	// FEditorViewportClient interface
@@ -49,8 +49,6 @@ protected:
 	// FEditorViewportClient interface
 	virtual void PerspectiveCameraMoved() override;
 
-	/** Call back for when the user changes preview scene settings in the UI */
-	void OnAssetViewerSettingsChanged(const FName& InPropertyName);
 	/** Used to (re)-set the viewport show flags related to post processing*/
 	void SetAdvancedShowFlagsForScene(const bool bAdvancedShowFlags);
 
@@ -63,7 +61,4 @@ private:
 
 	/** Pointer back to the PointCloudEditor viewport control that owns us */
 	TWeakPtr<SLidarPointCloudEditorViewport> PointCloudEditorViewportPtr;
-
-	/** Stored pointer to the preview scene in which the point cloud is shown */
-	FAdvancedPreviewScene* AdvancedPreviewScene;
 };
