Lidar - impossible to reset visibility
Panel terminology needs to be adjusted

diff --git a/Engine/Plugins/Enterprise/LidarPointCloud/Source/LidarPointCloudEditor/Private/LidarPointCloudEdModeToolkit.cpp b/Engine/Plugins/Enterprise/LidarPointCloud/Source/LidarPointCloudEditor/Private/LidarPointCloudEdModeToolkit.cpp
--- a/Engine/Plugins/Enterprise/LidarPointCloud/Source/LidarPointCloudEditor/Private/LidarPointCloudEdModeToolkit.cpp
+++ b/Engine/Plugins/Enterprise/LidarPointCloud/Source/LidarPointCloudEditor/Private/LidarPointCloudEdModeToolkit.cpp
@@ -33,11 +33,12 @@ void SLidarEditorWidget::Construct(const FArguments& InArgs)
 
 #define HEADEREX(ContText, Category) +SVerticalBox::Slot().AutoHeight().Padding(FMargin(0.0f, 5.0f))[SNew(SHeaderRow)+SHeaderRow::Column(ContText).HAlignCell(HAlign_Left).FillWidth(1).HeaderContentPadding(HeaderPadding)[SNew(STextBlock).Text(LOCTEXT(Category, ContText)).Font(HeaderFont)]]
 #define HEADER(ContText) HEADEREX(ContText, ContText"Header")
-#define BUTTON(Category, ConText, Action) +SHorizontalBox::Slot().Padding(StandardPadding).FillWidth(0.5f)[SNew(SButton).HAlign(HAlign_Center).Text(LOCTEXT(Category,ConText)).OnClicked_Lambda([this]{ Action return FReply::Handled(); })]
-#define BUTTON_ANY(Category, ConText, Action) +SHorizontalBox::Slot().Padding(StandardPadding).FillWidth(0.5f)[SNew(SButton).HAlign(HAlign_Center).Text(LOCTEXT(Category,ConText)).OnClicked_Lambda([this]{ Action return FReply::Handled(); }).IsEnabled(this, &SLidarEditorWidget::IsAnySelection)]
-#define BUTTON_POINTS(Category, ConText, Action) +SHorizontalBox::Slot().Padding(StandardPadding).FillWidth(0.5f)[SNew(SButton).HAlign(HAlign_Center).Text(LOCTEXT(Category,ConText)).OnClicked_Lambda([this]{ Action return FReply::Handled(); }).IsEnabled(this, &SLidarEditorWidget::IsPointSelection)]
-#define BUTTON_ACTORS(Category, ConText, Action) +SHorizontalBox::Slot().Padding(StandardPadding).FillWidth(0.5f)[SNew(SButton).HAlign(HAlign_Center).Text(LOCTEXT(Category,ConText)).OnClicked_Lambda([this]{ Action return FReply::Handled(); }).IsEnabled(this, &SLidarEditorWidget::IsActorSelection)]
+#define BUTTON(Category, Context, Tooltip, Action) +SHorizontalBox::Slot().Padding(StandardPadding).FillWidth(0.5f)[SNew(SButton).HAlign(HAlign_Center).Text(LOCTEXT(Category,Context)).OnClicked_Lambda([this]{ Action return FReply::Handled(); }).ToolTipText(LOCTEXT(Category,Tooltip))]
+#define BUTTON_ANY(Category, Context, Tooltip, Action) +SHorizontalBox::Slot().Padding(StandardPadding).FillWidth(0.5f)[SNew(SButton).HAlign(HAlign_Center).Text(LOCTEXT(Category,Context)).OnClicked_Lambda([this]{ Action return FReply::Handled(); }).IsEnabled(this, &SLidarEditorWidget::IsAnySelection).ToolTipText(LOCTEXT(Category,Tooltip))]
+#define BUTTON_POINTS(Category, Context, Tooltip, Action) +SHorizontalBox::Slot().Padding(StandardPadding).FillWidth(0.5f)[SNew(SButton).HAlign(HAlign_Center).Text(LOCTEXT(Category,Context)).OnClicked_Lambda([this]{ Action return FReply::Handled(); }).IsEnabled(this, &SLidarEditorWidget::IsPointSelection).ToolTipText(LOCTEXT(Category,Tooltip))]
+#define BUTTON_ACTORS(Category, Context, Tooltip, Action) +SHorizontalBox::Slot().Padding(StandardPadding).FillWidth(0.5f)[SNew(SButton).HAlign(HAlign_Center).Text(LOCTEXT(Category,Context)).OnClicked_Lambda([this]{ Action return FReply::Handled(); }).IsEnabled(this, &SLidarEditorWidget::IsActorSelection).ToolTipText(LOCTEXT(Category,Tooltip))]
 
+	
 	ChildSlot
 	[
 		SNew(SVerticalBox)
@@ -45,18 +46,12 @@ void SLidarEditorWidget::Construct(const FArguments& InArgs)
 		.AutoHeight()
 		[
 			SNew(SHorizontalBox)
-			BUTTON_ANY("SelectionClear", "Clear",
+			BUTTON_ANY("SelectionClear", "Clear", "Deselect all actors and points",
 			{
-				if(IsActorSelection())
-				{
-					FLidarPointCloudEditorHelper::ClearActorSelection();
-				}
-				else
-				{
-					FLidarPointCloudEditorHelper::ClearSelection();
-				}
+				FLidarPointCloudEditorHelper::ClearActorSelection();
+                FLidarPointCloudEditorHelper::ClearSelection();
 			})
-			BUTTON("SelectionInvert", "Invert",
+			BUTTON("SelectionInvert", "Invert", "Invert selected actors or points",
 			{
 				if(IsActorSelection())
 				{
@@ -127,24 +122,24 @@ void SLidarEditorWidget::Construct(const FArguments& InArgs)
 		.AutoHeight()
 		[
 			SNew(SHorizontalBox)
-			BUTTON_POINTS("CleanupHideSelected", "Hide Selected", {FLidarPointCloudEditorHelper::HideSelected();})
-			BUTTON_POINTS("CleanupResetVisibility", "Reset Visibility", {FLidarPointCloudEditorHelper::ResetVisibility();})
+			BUTTON_POINTS("CleanupHideSelected", "Hide Selected", "Hide all selected points", {FLidarPointCloudEditorHelper::HideSelected();})
+			BUTTON("CleanupResetVisibility", "Reset Visibility", "Mark all hidden points as visible", {FLidarPointCloudEditorHelper::ResetVisibility();})
 		]
 		+ SVerticalBox::Slot()
 		.AutoHeight()
 		[
 			SNew(SHorizontalBox)
-			BUTTON_POINTS("CleanupDeleteSelected", "Delete Selected", {FLidarPointCloudEditorHelper::DeleteSelected();})
-			BUTTON("CleanupDeleteHidden", "Delete Hidden", {FLidarPointCloudEditorHelper::DeleteHidden();})
+			BUTTON_POINTS("CleanupDeleteSelected", "Delete Selected", "Permanently delete all selected points", {FLidarPointCloudEditorHelper::DeleteSelected();})
+			BUTTON("CleanupDeleteHidden", "Delete Hidden", "Permanently delete all hidden points", {FLidarPointCloudEditorHelper::DeleteHidden();})
 		]
 		+ SVerticalBox::Slot()
 		.AutoHeight()
 		[
 			SNew(SHorizontalBox)
-			BUTTON_POINTS("CleanupCropToSelection", "Crop To Selection",
+			BUTTON_POINTS("CleanupCropToSelection", "Crop To Selection", "Hides all points not within selection",
 			{
 				FLidarPointCloudEditorHelper::InvertSelection();
-				FLidarPointCloudEditorHelper::DeleteSelected();
+				FLidarPointCloudEditorHelper::HideSelected();
 			})
 			+ SHorizontalBox::Slot()
 			.Padding(StandardPadding)
@@ -158,6 +153,7 @@ void SLidarEditorWidget::Construct(const FArguments& InArgs)
 		.AutoHeight()
 		[
 			SNew(SHorizontalBox)
+			.ToolTipText(LOCTEXT("CollisionsError", "Determines the maximum error (in cm) of the collision. Lower values will require more time to build."))
 			+ SHorizontalBox::Slot()
 			.Padding(StandardPadding)
 			.FillWidth(0.5f)
@@ -186,18 +182,19 @@ void SLidarEditorWidget::Construct(const FArguments& InArgs)
 		.AutoHeight()
 		[
 			SNew(SHorizontalBox)
-			BUTTON_ACTORS("CollisionsAdd", "Add Collision",
+			BUTTON_ACTORS("CollisionsAdd", "Add Collision", "Generate collision for selected assets",
 			{
 				FLidarPointCloudEditorHelper::SetCollisionErrorForSelection(MaxCollisionError);
 				FLidarPointCloudEditorHelper::BuildCollisionForSelection();
 			})
-			BUTTON_ACTORS("CollisionsRemove", "Remove Collision", {FLidarPointCloudEditorHelper::RemoveCollisionForSelection();})
+			BUTTON_ACTORS("CollisionsRemove", "Remove Collision", "Removes any existing collisions from selected assets", {FLidarPointCloudEditorHelper::RemoveCollisionForSelection();})
 		]
 		HEADER("Normals")
 		+ SVerticalBox::Slot()
 		.AutoHeight()
 		[
 			SNew(SHorizontalBox)
+			.ToolTipText(LOCTEXT("NormalsQuality", "Higher values will generally result in more accurate calculations, at the expense of time"))
 			+ SHorizontalBox::Slot()
 			.Padding(StandardPadding)
 			.FillWidth(0.5f)
@@ -223,6 +220,7 @@ void SLidarEditorWidget::Construct(const FArguments& InArgs)
 		.AutoHeight()
 		[
 			SNew(SHorizontalBox)
+			.ToolTipText(LOCTEXT("NormalsNoiseTolerance", "Higher values are less susceptible to noise, but will most likely lose finer details, especially around hard edges."))
 			+ SHorizontalBox::Slot()
 			.Padding(StandardPadding)
 			.FillWidth(0.5f)
@@ -250,7 +248,7 @@ void SLidarEditorWidget::Construct(const FArguments& InArgs)
 		.AutoHeight()
 		[
 			SNew(SHorizontalBox)
-			BUTTON_ANY("NormalsCalculate", "Calculate Normals",
+			BUTTON_ANY("NormalsCalculate", "Calculate Normals", "Calculates normal vectors for all selected assets or points",
 			{
 				FLidarPointCloudEditorHelper::SetNormalsQuality(NormalsQuality, NormalsNoiseTolerance);
 				if(IsPointSelection())
@@ -268,6 +266,7 @@ void SLidarEditorWidget::Construct(const FArguments& InArgs)
 		.AutoHeight()
 		[
 			SNew(SHorizontalBox)
+			.ToolTipText(LOCTEXT("MeshingError", "Determines the maximum error (in cm) of the resulting mesh. Lower values will require more time to build."))
 			+ SHorizontalBox::Slot()
 			.Padding(StandardPadding)
 			.FillWidth(0.5f)
@@ -295,6 +294,7 @@ void SLidarEditorWidget::Construct(const FArguments& InArgs)
 		.AutoHeight()
 		[
 			SNew(SHorizontalBox)
+			.ToolTipText(LOCTEXT("MeshingMerge", "When enabled, all elements will be combined into one big mesh"))
 			+ SHorizontalBox::Slot()
 			.Padding(StandardPadding)
 			.FillWidth(0.5f)
@@ -316,6 +316,7 @@ void SLidarEditorWidget::Construct(const FArguments& InArgs)
 		.AutoHeight()
 		[
 			SNew(SHorizontalBox)
+			.ToolTipText(LOCTEXT("MeshingRetain", "When enabled, all generated meshes will retain the transforms of their source lidar data"))
 			+ SHorizontalBox::Slot()
 			.Padding(StandardPadding)
 			.FillWidth(0.5f)
@@ -338,21 +339,23 @@ void SLidarEditorWidget::Construct(const FArguments& InArgs)
 		.AutoHeight()
 		[
 			SNew(SHorizontalBox)
-			BUTTON_ANY("MeshingBuild", "Create Static Mesh", {FLidarPointCloudEditorHelper::MeshSelected(IsPointSelection(), MaxMeshingError, bMergeMeshes, !bMergeMeshes && bRetainTransform);})
+			BUTTON_ANY("MeshingBuild", "Create Static Mesh", "Generates a static mesh from selected assets or points", {
+				FLidarPointCloudEditorHelper::MeshSelected(IsPointSelection(), MaxMeshingError, bMergeMeshes, !bMergeMeshes && bRetainTransform);
+			})
 		]
 		HEADER("Alignment")
 		+ SVerticalBox::Slot()
 		.AutoHeight()
 		[
 			SNew(SHorizontalBox)
-			BUTTON_ACTORS("AlignmentCoords", "Original Coordinates", {FLidarPointCloudEditorHelper::SetOriginalCoordinateForSelection();})
-			BUTTON_ACTORS("AlignmentOrigin", "World Origin", {FLidarPointCloudEditorHelper::AlignSelectionAroundWorldOrigin();})
+			BUTTON_ACTORS("AlignmentCoords", "Original Coordinates", "Aligns all selected assets to their original coordinates", {FLidarPointCloudEditorHelper::SetOriginalCoordinateForSelection();})
+			BUTTON_ACTORS("AlignmentOrigin", "World Origin", "Aligns all selected assets relative to each other, while keeping the overall pivot around 0,0,0", {FLidarPointCloudEditorHelper::AlignSelectionAroundWorldOrigin();})
 		]
 		+ SVerticalBox::Slot()
 		.AutoHeight()
 		[
 			SNew(SHorizontalBox)
-			BUTTON_ACTORS("AlignmentReset", "Reset Alignment", {FLidarPointCloudEditorHelper::CenterSelection();})
+			BUTTON_ACTORS("AlignmentReset", "Reset Alignment", "Resets the pivot of all selected assets to 0,0,0", {FLidarPointCloudEditorHelper::CenterSelection();})
 			+ SHorizontalBox::Slot()
 			.Padding(StandardPadding)
 			.FillWidth(0.5f)
@@ -365,15 +368,15 @@ void SLidarEditorWidget::Construct(const FArguments& InArgs)
 		.AutoHeight()
 		[
 			SNew(SHorizontalBox)
-			BUTTON_POINTS("MergeExtractSelection", "Extract Selection", {FLidarPointCloudEditorHelper::Extract();})
-			BUTTON_POINTS("MergeExtractCopy", "Extract as Copy", {FLidarPointCloudEditorHelper::ExtractAsCopy();})
+			BUTTON_POINTS("MergeExtractSelection", "Extract Selection", "Extracts the selected points as separate assets, removing them from the original one", {FLidarPointCloudEditorHelper::Extract();})
+			BUTTON_POINTS("MergeExtractCopy", "Extract as Copy", "Extracts the selected points as separate assets, retaining them in the original one", {FLidarPointCloudEditorHelper::ExtractAsCopy();})
 		]
 		+ SVerticalBox::Slot()
 		.AutoHeight()
 		[
 			SNew(SHorizontalBox)
-			BUTTON_ACTORS("MergeExtractActors", "Merge Actors", {FLidarPointCloudEditorHelper::MergeSelectionByComponent(true);})
-			BUTTON_ACTORS("MergeExtractData", "Merge Data", {FLidarPointCloudEditorHelper::MergeSelectionByData(true);})
+			BUTTON_ACTORS("MergeExtractActors", "Merge Actors", "Replaces all individual selected assets with a single new actor", {FLidarPointCloudEditorHelper::MergeSelectionByComponent(true);})
+			BUTTON_ACTORS("MergeExtractAssets", "Merge Assets", "Combines all selected assets into a single, large new asset", {FLidarPointCloudEditorHelper::MergeSelectionByData(true);})
 		]
 	];
 
@@ -395,6 +398,11 @@ bool SLidarEditorWidget::IsPointSelection() const
 	return !bActorSelection && FLidarPointCloudEditorHelper::AreLidarPointsSelected();
 }
 
+bool SLidarEditorWidget::IsAnySelection() const
+{
+	return FLidarPointCloudEditorHelper::AreLidarActorsSelected() || FLidarPointCloudEditorHelper::AreLidarPointsSelected();
+}
+
 FLidarPointCloudEdModeToolkit::~FLidarPointCloudEdModeToolkit()
 {
 	GetScriptableEditorMode()->GetInteractiveToolsContext(EToolsContextScope::EdMode)->OnToolNotificationMessage.RemoveAll(this);
diff --git a/Engine/Plugins/Enterprise/LidarPointCloud/Source/LidarPointCloudEditor/Private/LidarPointCloudEdModeToolkit.h b/Engine/Plugins/Enterprise/LidarPointCloud/Source/LidarPointCloudEditor/Private/LidarPointCloudEdModeToolkit.h
--- a/Engine/Plugins/Enterprise/LidarPointCloud/Source/LidarPointCloudEditor/Private/LidarPointCloudEdModeToolkit.h
+++ b/Engine/Plugins/Enterprise/LidarPointCloud/Source/LidarPointCloudEditor/Private/LidarPointCloudEdModeToolkit.h
@@ -24,7 +24,7 @@ public:
 
 	bool IsActorSelection() const;
 	bool IsPointSelection() const;
-	bool IsAnySelection() const { return IsActorSelection() || IsPointSelection(); }
+	bool IsAnySelection() const;
 	bool IsBrushVisible() const { return BrushTool != nullptr; }
 
 	EVisibility GetActorVisibility() const { return IsActorSelection() ? EVisibility::Visible : EVisibility::Collapsed; }
