Fixed clipping volumes are not selectable in Lidar mode

diff --git a/Engine/Plugins/Enterprise/LidarPointCloud/Source/LidarPointCloudEditor/Private/LidarPointCloudEdMode.cpp b/Engine/Plugins/Enterprise/LidarPointCloud/Source/LidarPointCloudEditor/Private/LidarPointCloudEdMode.cpp
--- a/Engine/Plugins/Enterprise/LidarPointCloud/Source/LidarPointCloudEditor/Private/LidarPointCloudEdMode.cpp
+++ b/Engine/Plugins/Enterprise/LidarPointCloud/Source/LidarPointCloudEditor/Private/LidarPointCloudEdMode.cpp
@@ -6,6 +6,7 @@
 #include "EditorModeManager.h"
 #include "InteractiveToolManager.h"
 #include "LevelEditorViewport.h"
+#include "LidarPointCloud.h"
 #include "LidarPointCloudActor.h"
 #include "LidarPointCloudEditorCommands.h"
 #include "LidarPointCloudEditorTools.h"
@@ -52,7 +53,7 @@ void ULidarEditorMode::Enter()
 
 bool ULidarEditorMode::IsSelectionAllowed(AActor* InActor, bool bInSelection) const
 {
-	return InActor && InActor->IsA(ALidarPointCloudActor::StaticClass());
+	return InActor && (InActor->IsA(ALidarPointCloudActor::StaticClass()) || InActor->IsA(ALidarClippingVolume::StaticClass()));
 }
 
 void ULidarEditorMode::CreateToolkit()
