Fixed Extracted actor in level doesn't carry the asset name
Fixed Cancel point extraction result in editor crash

diff --git a/Engine/Plugins/Enterprise/LidarPointCloud/Source/LidarPointCloudEditor/Private/LidarPointCloudEditorHelper.cpp b/Engine/Plugins/Enterprise/LidarPointCloud/Source/LidarPointCloudEditor/Private/LidarPointCloudEditorHelper.cpp
--- a/Engine/Plugins/Enterprise/LidarPointCloud/Source/LidarPointCloudEditor/Private/LidarPointCloudEditorHelper.cpp
+++ b/Engine/Plugins/Enterprise/LidarPointCloud/Source/LidarPointCloudEditor/Private/LidarPointCloudEditorHelper.cpp
@@ -156,12 +156,15 @@ namespace
 	}
 
 	template<typename T>
-	T* SpawnActor()
+	T* SpawnActor(const FString& Name)
 	{
 		T* NewActor = nullptr;
 		if(UWorld* World = GetFirstWorld())
 		{
-			NewActor = Cast<T>(World->SpawnActor(T::StaticClass()));
+			FActorSpawnParameters ActorSpawnParameters;
+			ActorSpawnParameters.Name = FName(*Name);
+			NewActor = Cast<T>(World->SpawnActor<T>(ActorSpawnParameters));
+			FActorLabelUtilities::SetActorLabelUnique(NewActor, NewActor->GetName());
 		}
 		return NewActor;
 	}
@@ -370,22 +373,22 @@ namespace
 	
 	ULidarPointCloud* Extract_Internal()
 	{
-		int64 NumPoints = 0;
-		ProcessAll([&NumPoints](ALidarPointCloudActor* Actor)
-		{
-			NumPoints += Actor->GetPointCloudComponent()->NumSelectedPoints();
-		});
-
-		TArray64<FLidarPointCloudPoint> SelectedPoints;
-		SelectedPoints.Reserve(NumPoints);
-		ProcessAll([&SelectedPoints](ALidarPointCloudActor* Actor)
-		{
-			Actor->GetPointCloudComponent()->GetSelectedPointsAsCopies(SelectedPoints);
-		});
-	
 		ULidarPointCloud* NewPointCloud = CreateNewAsset_Internal();
 		if (NewPointCloud)
 		{
+			int64 NumPoints = 0;
+			ProcessAll([&NumPoints](ALidarPointCloudActor* Actor)
+			{
+				NumPoints += Actor->GetPointCloudComponent()->NumSelectedPoints();
+			});
+
+			TArray64<FLidarPointCloudPoint> SelectedPoints;
+			SelectedPoints.Reserve(NumPoints);
+			ProcessAll([&SelectedPoints](ALidarPointCloudActor* Actor)
+			{
+				Actor->GetPointCloudComponent()->GetSelectedPointsAsCopies(SelectedPoints);
+			});
+			
 			NewPointCloud->SetData(SelectedPoints);
 		}
 		return NewPointCloud;
@@ -535,7 +538,7 @@ void FLidarPointCloudEditorHelper::MeshSelected(bool bMeshByPoints, float CellSi
 			AssignMeshBuffersToMesh(&MeshBuffers, StaticMesh);
 			ProgressDialog.EnterProgressFrame(1.0f);
 					
-			if(AStaticMeshActor* MeshActor = SpawnActor<AStaticMeshActor>())
+			if(AStaticMeshActor* MeshActor = SpawnActor<AStaticMeshActor>(StaticMesh->GetName()))
 			{
 				MeshActor->GetStaticMeshComponent()->SetStaticMesh(StaticMesh);
 			}
@@ -562,7 +565,7 @@ void FLidarPointCloudEditorHelper::MeshSelected(bool bMeshByPoints, float CellSi
 					
 					AssignMeshBuffersToMesh(&MeshBuffers, StaticMesh);
 					
-					if(AStaticMeshActor* MeshActor = SpawnActor<AStaticMeshActor>())
+					if(AStaticMeshActor* MeshActor = SpawnActor<AStaticMeshActor>(StaticMesh->GetName()))
 					{
 						if(bRetainTransform)
 						{
@@ -635,26 +638,30 @@ void FLidarPointCloudEditorHelper::DeleteHidden()
 
 void FLidarPointCloudEditorHelper::Extract()
 {
-	ULidarPointCloud* NewPointCloud = Extract_Internal();
-	if(ALidarPointCloudActor* Actor = SpawnActor<ALidarPointCloudActor>())
+	if(ULidarPointCloud* NewPointCloud = Extract_Internal())
 	{
-		NewPointCloud->RestoreOriginalCoordinates();
-		Actor->SetPointCloud(NewPointCloud);
-	}
+		if(ALidarPointCloudActor* Actor = SpawnActor<ALidarPointCloudActor>(NewPointCloud->GetName()))
+		{
+			NewPointCloud->RestoreOriginalCoordinates();
+			Actor->SetPointCloud(NewPointCloud);
+		}
 
-	DeleteSelected();
+		DeleteSelected();
+	}
 }
 
 void FLidarPointCloudEditorHelper::ExtractAsCopy()
 {
-	ULidarPointCloud* NewPointCloud = Extract_Internal();
-	if(ALidarPointCloudActor* Actor = SpawnActor<ALidarPointCloudActor>())
+	if(ULidarPointCloud* NewPointCloud = Extract_Internal())
 	{
-		NewPointCloud->RestoreOriginalCoordinates();
-		Actor->SetPointCloud(NewPointCloud);
-	}
+		if(ALidarPointCloudActor* Actor = SpawnActor<ALidarPointCloudActor>(NewPointCloud->GetName()))
+		{
+			NewPointCloud->RestoreOriginalCoordinates();
+			Actor->SetPointCloud(NewPointCloud);
+		}
 	
-	ClearSelection();
+		ClearSelection();
+	}
 }
 
 void FLidarPointCloudEditorHelper::CalculateNormals()
