Lidar Icon redesign

diff --git a/Engine/Plugins/Enterprise/LidarPointCloud/Source/LidarPointCloudEditor/Private/LidarPointCloudFactory.cpp b/Engine/Plugins/Enterprise/LidarPointCloud/Source/LidarPointCloudEditor/Private/LidarPointCloudFactory.cpp
--- a/Engine/Plugins/Enterprise/LidarPointCloud/Source/LidarPointCloudEditor/Private/LidarPointCloudFactory.cpp
+++ b/Engine/Plugins/Enterprise/LidarPointCloud/Source/LidarPointCloudEditor/Private/LidarPointCloudFactory.cpp
@@ -59,7 +59,7 @@ void FAssetTypeActions_LidarPointCloud::GetActions(const TArray<UObject*>& InObj
 		MenuBuilder.AddMenuEntry(
 			LOCTEXT("LidarPointCloud_Merge", "Merge Selected"),
 			LOCTEXT("LidarPointCloud_MergeTooltip", "Merges selected point cloud assets."),
-			FSlateIcon("LidarPointCloudStyle", "LidarPointCloudEditor.Merge"),
+			FSlateIcon("LidarPointCloudStyle", "LidarPointCloudEditor.ToolkitMerge"),
 			FUIAction(
 				FExecuteAction::CreateSP(this, &FAssetTypeActions_LidarPointCloud::ExecuteMerge, PointClouds),
 				FCanExecuteAction()
@@ -69,7 +69,7 @@ void FAssetTypeActions_LidarPointCloud::GetActions(const TArray<UObject*>& InObj
 		MenuBuilder.AddMenuEntry(
 			LOCTEXT("LidarPointCloud_Align", "Align Selected"),
 			LOCTEXT("LidarPointCloud_AlignTooltip", "Aligns selected point cloud assets."),
-			FSlateIcon("LidarPointCloudStyle", "LidarPointCloudEditor.Align"),
+			FSlateIcon("LidarPointCloudStyle", "LidarPointCloudEditor.ToolkitAlign"),
 			FUIAction(
 				FExecuteAction::CreateSP(this, &FAssetTypeActions_LidarPointCloud::ExecuteAlign, PointClouds),
 				FCanExecuteAction()
@@ -80,7 +80,7 @@ void FAssetTypeActions_LidarPointCloud::GetActions(const TArray<UObject*>& InObj
 	MenuBuilder.AddMenuEntry(
 		LOCTEXT("LidarPointCloud_BuildCollision", "Build Collision"),
 		LOCTEXT("LidarPointCloud_BuildCollisionTooltip", "Builds collision for all selected point cloud assets."),
-		FSlateIcon("LidarPointCloudStyle", "LidarPointCloudEditor.BuildCollision"),
+		FSlateIcon("LidarPointCloudStyle", "LidarPointCloudEditor.ToolkitCollision"),
 		FUIAction(
 			FExecuteAction::CreateSP(this, &FAssetTypeActions_LidarPointCloud::ExecuteCollision, PointClouds),
 			FCanExecuteAction()
@@ -90,7 +90,7 @@ void FAssetTypeActions_LidarPointCloud::GetActions(const TArray<UObject*>& InObj
 	MenuBuilder.AddMenuEntry(
 		LOCTEXT("LidarPointCloud_CalculateNormals", "Calculate Normals"),
 		LOCTEXT("LidarPointCloud_CalculateNormalsTooltip", "Calculates normals for all selected point cloud assets."),
-		FSlateIcon(FAppStyle::GetAppStyleSetName(), "AnimViewportMenu.SetShowNormals"),
+		FSlateIcon("LidarPointCloudStyle", "LidarPointCloudEditor.ToolkitNormals"),
 		FUIAction(
 			FExecuteAction::CreateSP(this, &FAssetTypeActions_LidarPointCloud::ExecuteNormals, PointClouds),
 			FCanExecuteAction()
diff --git a/Engine/Plugins/Enterprise/LidarPointCloud/Source/LidarPointCloudEditor/Private/LidarPointCloudStyle.cpp b/Engine/Plugins/Enterprise/LidarPointCloud/Source/LidarPointCloudEditor/Private/LidarPointCloudStyle.cpp
--- a/Engine/Plugins/Enterprise/LidarPointCloud/Source/LidarPointCloudEditor/Private/LidarPointCloudStyle.cpp
+++ b/Engine/Plugins/Enterprise/LidarPointCloud/Source/LidarPointCloudEditor/Private/LidarPointCloudStyle.cpp
@@ -8,6 +8,7 @@
 #include "Interfaces/IPluginManager.h"
 
 #define IMAGE_PLUGIN_BRUSH( RelativePath, ... ) FSlateImageBrush( StyleSet->RootToContentDir( RelativePath, TEXT(".png") ), __VA_ARGS__ )
+#define IMAGE_PLUGIN_BRUSH_SVG( RelativePath, ... ) FSlateVectorImageBrush( StyleSet->RootToContentDir( RelativePath, TEXT(".svg") ), __VA_ARGS__ )
 
 TSharedPtr<FSlateStyleSet> FLidarPointCloudStyle::StyleSet = nullptr;
 TSharedPtr<class ISlateStyle> FLidarPointCloudStyle::Get() { return StyleSet; }
@@ -20,7 +21,7 @@ void FLidarPointCloudStyle::Initialize()
 	const FVector2D Icon32x32(32.0f, 32.0f);
 	const FVector2D Icon40x40(40.0f, 40.0f);
 	const FVector2D Icon128x128(128.0f, 128.0f);
-
+	
 	// Only register once
 	if (StyleSet.IsValid())
 	{
@@ -44,41 +45,19 @@ void FLidarPointCloudStyle::Initialize()
 	StyleSet->Set("ClassIcon.LidarClippingVolume", new IMAGE_PLUGIN_BRUSH("icon_32", Icon16x16));
 	StyleSet->Set("ClassIcon32.LidarClippingVolume", new IMAGE_PLUGIN_BRUSH("icon_32", Icon32x32));
 	StyleSet->Set("ClassThumbnail.LidarClippingVolume", new IMAGE_PLUGIN_BRUSH("icon_128", Icon128x128));
-
-	StyleSet->Set("LidarPointCloudEditor.BuildCollision", new IMAGE_PLUGIN_BRUSH("icon_collision_40", Icon40x40));
-	StyleSet->Set("LidarPointCloudEditor.RemoveCollision", new IMAGE_PLUGIN_BRUSH("icon_removecollision_40", Icon40x40));
-
-	StyleSet->Set("LidarPointCloudEditor.EditMode", new IMAGE_PLUGIN_BRUSH("icon_edit_40", Icon40x40));
-	
-	StyleSet->Set("LidarPointCloudEditor.BoxSelection", new IMAGE_PLUGIN_BRUSH("icon_selbox_40", Icon40x40));
-	StyleSet->Set("LidarPointCloudEditor.PolygonalSelection", new IMAGE_PLUGIN_BRUSH("icon_selpoly_40", Icon40x40));
-	StyleSet->Set("LidarPointCloudEditor.LassoSelection", new IMAGE_PLUGIN_BRUSH("icon_sellasso_40", Icon40x40));
-	StyleSet->Set("LidarPointCloudEditor.PaintSelection", new IMAGE_PLUGIN_BRUSH("icon_selpaint_40", Icon40x40));
-
-	StyleSet->Set("LidarPointCloudEditor.InvertSelection", new IMAGE_PLUGIN_BRUSH("icon_invsel_40", Icon40x40));
-
-	StyleSet->Set("LidarPointCloudEditor.HideSelected", new IMAGE_PLUGIN_BRUSH("icon_hideselected_40", Icon40x40));
-	StyleSet->Set("LidarPointCloudEditor.UnhideAll", new IMAGE_PLUGIN_BRUSH("icon_unhideall_40", Icon40x40));
-
-	StyleSet->Set("LidarPointCloudEditor.DeleteSelected", new IMAGE_PLUGIN_BRUSH("icon_deleteselected_40", Icon40x40));
-
-	StyleSet->Set("LidarPointCloudEditor.Extract", new IMAGE_PLUGIN_BRUSH("icon_extract_40", Icon40x40));
-	StyleSet->Set("LidarPointCloudEditor.ExtractCopy", new IMAGE_PLUGIN_BRUSH("icon_extractcopy_40", Icon40x40));
-	StyleSet->Set("LidarPointCloudEditor.Merge", new IMAGE_PLUGIN_BRUSH("icon_merge_40", Icon40x40));
-	StyleSet->Set("LidarPointCloudEditor.Align", new IMAGE_PLUGIN_BRUSH("icon_align_40", Icon40x40));
 	
-	StyleSet->Set("LidarPointCloudEditor.ToolkitCollision", new IMAGE_PLUGIN_BRUSH("icon_collision_40", Icon20x20));
-	StyleSet->Set("LidarPointCloudEditor.ToolkitMerge", new IMAGE_PLUGIN_BRUSH("icon_merge_40", Icon20x20));
-	StyleSet->Set("LidarPointCloudEditor.ToolkitAlign", new IMAGE_PLUGIN_BRUSH("icon_align_40", Icon20x20));
-	StyleSet->Set("LidarPointCloudEditor.ToolkitBoxSelection", new IMAGE_PLUGIN_BRUSH("icon_selbox_40", Icon20x20));
-	StyleSet->Set("LidarPointCloudEditor.ToolkitPolygonalSelection", new IMAGE_PLUGIN_BRUSH("icon_selpoly_40", Icon20x20));
-	StyleSet->Set("LidarPointCloudEditor.ToolkitLassoSelection", new IMAGE_PLUGIN_BRUSH("icon_sellasso_40", Icon20x20));
-	StyleSet->Set("LidarPointCloudEditor.ToolkitPaintSelection", new IMAGE_PLUGIN_BRUSH("icon_selpaint_40", Icon20x20));
+	StyleSet->Set("LidarPointCloudEditor.ToolkitCollision", new IMAGE_PLUGIN_BRUSH_SVG("Collision_20", Icon20x20));
+	StyleSet->Set("LidarPointCloudEditor.ToolkitMerge", new IMAGE_PLUGIN_BRUSH_SVG("Merge_20", Icon20x20));
+	StyleSet->Set("LidarPointCloudEditor.ToolkitAlign", new IMAGE_PLUGIN_BRUSH_SVG("AlignLidar", Icon20x20));
+	StyleSet->Set("LidarPointCloudEditor.ToolkitBoxSelection", new IMAGE_PLUGIN_BRUSH_SVG("BoxSelection", Icon20x20));
+	StyleSet->Set("LidarPointCloudEditor.ToolkitPolygonalSelection", new IMAGE_PLUGIN_BRUSH_SVG("PolySelection", Icon20x20));
+	StyleSet->Set("LidarPointCloudEditor.ToolkitLassoSelection", new IMAGE_PLUGIN_BRUSH_SVG("Lasso", Icon20x20));
+	StyleSet->Set("LidarPointCloudEditor.ToolkitNormals", new IMAGE_PLUGIN_BRUSH_SVG( TEXT("SetShowNormals_20"), Icon20x20 ));
 	
 	StyleSet->SetContentRoot(FPaths::EngineContentDir() / TEXT("Editor/Slate"));
 	
-	StyleSet->Set("LidarPointCloudEditor.ToolkitSelect", new IMAGE_PLUGIN_BRUSH( TEXT("Icons/GeneralTools/Select_40x"), Icon20x20 ) );
-	StyleSet->Set("LidarPointCloudEditor.ToolkitNormals", new IMAGE_PLUGIN_BRUSH( TEXT("Icons/icon_StaticMeshEd_Normals_40x"), Icon20x20 ) );
+	StyleSet->Set("LidarPointCloudEditor.ToolkitSelect", new IMAGE_PLUGIN_BRUSH( TEXT("Icons/GeneralTools/Select_40x"), Icon20x20 ));
+	StyleSet->Set("LidarPointCloudEditor.ToolkitPaintSelection", new IMAGE_PLUGIN_BRUSH_SVG( TEXT("Starship/Common/Paintbrush"), Icon20x20));
 
 	FSlateStyleRegistry::RegisterSlateStyle(*StyleSet.Get());
 }
