Lidar - Have 'Import with new file' in the asset contextual menu

diff --git a/Engine/Plugins/Enterprise/LidarPointCloud/Source/LidarPointCloudEditor/Private/LidarPointCloudFactory.cpp b/Engine/Plugins/Enterprise/LidarPointCloud/Source/LidarPointCloudEditor/Private/LidarPointCloudFactory.cpp
--- a/Engine/Plugins/Enterprise/LidarPointCloud/Source/LidarPointCloudEditor/Private/LidarPointCloudFactory.cpp
+++ b/Engine/Plugins/Enterprise/LidarPointCloud/Source/LidarPointCloudEditor/Private/LidarPointCloudFactory.cpp
@@ -30,30 +30,6 @@ void FAssetTypeActions_LidarPointCloud::GetActions(const TArray<UObject*>& InObj
 		PointClouds.Add(CastChecked<ULidarPointCloud>(Object));
 	}
 
-	// Make sure at least one asset has source assigned
-	bool bSourceExists = false;
-	for (ULidarPointCloud* PointCloud : PointClouds)
-	{
-		if (!PointCloud->GetSourcePath().IsEmpty())
-		{
-			bSourceExists = true;
-			break;
-		}
-	}
-
-	if (bSourceExists)
-	{
-		MenuBuilder.AddMenuEntry(
-			LOCTEXT("LidarPointCloud_Reimport", "Reimport Selected"),
-			LOCTEXT("LidarPointCloud_ReimportTooltip", "Reimports this point cloud asset."),
-			FSlateIcon(FAppStyle::GetAppStyleSetName(), "TextureEditor.Reimport"),
-			FUIAction(
-				FExecuteAction::CreateSP(this, &FAssetTypeActions_LidarPointCloud::ExecuteReimport, PointClouds),
-				FCanExecuteAction()
-			)
-		);
-	}
-
 	if (PointClouds.Num() > 1)
 	{
 		MenuBuilder.AddMenuEntry(
@@ -112,11 +88,14 @@ void FAssetTypeActions_LidarPointCloud::OpenAssetEditor(const TArray<UObject*>&
 	}
 }
 
-void FAssetTypeActions_LidarPointCloud::ExecuteReimport(TArray<ULidarPointCloud*> PointClouds)
+void FAssetTypeActions_LidarPointCloud::GetResolvedSourceFilePaths(const TArray<UObject*>& TypeAssets, TArray<FString>& OutSourceFilePaths) const
 {
-	for (ULidarPointCloud* PC : PointClouds)
+	for (auto& Asset : TypeAssets)
 	{
-		PC->Reimport(GetDefault<ULidarPointCloudSettings>()->bUseAsyncImport);
+		if (const ULidarPointCloud* PointCloud = CastChecked<ULidarPointCloud>(Asset))
+		{
+			OutSourceFilePaths.Add(PointCloud->GetSourcePath());
+		}
 	}
 }
 
diff --git a/Engine/Plugins/Enterprise/LidarPointCloud/Source/LidarPointCloudEditor/Private/LidarPointCloudFactory.h b/Engine/Plugins/Enterprise/LidarPointCloud/Source/LidarPointCloudEditor/Private/LidarPointCloudFactory.h
--- a/Engine/Plugins/Enterprise/LidarPointCloud/Source/LidarPointCloudEditor/Private/LidarPointCloudFactory.h
+++ b/Engine/Plugins/Enterprise/LidarPointCloud/Source/LidarPointCloudEditor/Private/LidarPointCloudFactory.h
@@ -20,11 +20,12 @@ public:
 	virtual bool HasActions(const TArray<UObject*>& InObjects) const override { return true; }
 	virtual void GetActions(const TArray<UObject*>& InObjects, FMenuBuilder& MenuBuilder) override;
 	virtual uint32 GetCategories() override { return EAssetTypeCategories::Misc; }
+	virtual bool IsImportedAsset() const override { return true; }
 	virtual void OpenAssetEditor(const TArray<UObject*>& InObjects, TSharedPtr<class IToolkitHost> EditWithinLevelEditor = TSharedPtr<IToolkitHost>()) override;
+	virtual void GetResolvedSourceFilePaths(const TArray<UObject*>& TypeAssets, TArray<FString>& OutSourceFilePaths) const override;
 	// End IAssetTypeActions Interface
 
 private:
-	void ExecuteReimport(TArray<ULidarPointCloud*> PointClouds);
 	void ExecuteMerge(TArray<ULidarPointCloud*> PointClouds);
 	void ExecuteAlign(TArray<ULidarPointCloud*> PointClouds);
 	void ExecuteCollision(TArray<ULidarPointCloud*> PointClouds);
