Fixed text conflict from LOCTEXT macro

diff --git a/Engine/Plugins/Enterprise/LidarPointCloud/Source/LidarPointCloudEditor/Private/LidarPointCloudEdModeToolkit.cpp b/Engine/Plugins/Enterprise/LidarPointCloud/Source/LidarPointCloudEditor/Private/LidarPointCloudEdModeToolkit.cpp
--- a/Engine/Plugins/Enterprise/LidarPointCloud/Source/LidarPointCloudEditor/Private/LidarPointCloudEdModeToolkit.cpp
+++ b/Engine/Plugins/Enterprise/LidarPointCloud/Source/LidarPointCloudEditor/Private/LidarPointCloudEdModeToolkit.cpp
@@ -33,11 +33,10 @@ void SLidarEditorWidget::Construct(const FArguments& InArgs)
 
 #define HEADEREX(ContText, Category) +SVerticalBox::Slot().AutoHeight().Padding(FMargin(0.0f, 5.0f))[SNew(SHeaderRow)+SHeaderRow::Column(ContText).HAlignCell(HAlign_Left).FillWidth(1).HeaderContentPadding(HeaderPadding)[SNew(STextBlock).Text(LOCTEXT(Category, ContText)).Font(HeaderFont)]]
 #define HEADER(ContText) HEADEREX(ContText, ContText"Header")
-#define BUTTON(Category, Context, Tooltip, Action) +SHorizontalBox::Slot().Padding(StandardPadding).FillWidth(0.5f)[SNew(SButton).HAlign(HAlign_Center).Text(LOCTEXT(Category,Context)).OnClicked_Lambda([this]{ Action return FReply::Handled(); }).ToolTipText(LOCTEXT(Category,Tooltip))]
-#define BUTTON_ANY(Category, Context, Tooltip, Action) +SHorizontalBox::Slot().Padding(StandardPadding).FillWidth(0.5f)[SNew(SButton).HAlign(HAlign_Center).Text(LOCTEXT(Category,Context)).OnClicked_Lambda([this]{ Action return FReply::Handled(); }).IsEnabled(this, &SLidarEditorWidget::IsAnySelection).ToolTipText(LOCTEXT(Category,Tooltip))]
-#define BUTTON_POINTS(Category, Context, Tooltip, Action) +SHorizontalBox::Slot().Padding(StandardPadding).FillWidth(0.5f)[SNew(SButton).HAlign(HAlign_Center).Text(LOCTEXT(Category,Context)).OnClicked_Lambda([this]{ Action return FReply::Handled(); }).IsEnabled(this, &SLidarEditorWidget::IsPointSelection).ToolTipText(LOCTEXT(Category,Tooltip))]
-#define BUTTON_ACTORS(Category, Context, Tooltip, Action) +SHorizontalBox::Slot().Padding(StandardPadding).FillWidth(0.5f)[SNew(SButton).HAlign(HAlign_Center).Text(LOCTEXT(Category,Context)).OnClicked_Lambda([this]{ Action return FReply::Handled(); }).IsEnabled(this, &SLidarEditorWidget::IsActorSelection).ToolTipText(LOCTEXT(Category,Tooltip))]
-
+#define BUTTON(Category, Context, Tooltip, Action) +SHorizontalBox::Slot().Padding(StandardPadding).FillWidth(0.5f)[SNew(SButton).HAlign(HAlign_Center).Text(LOCTEXT(Category,Context)).OnClicked_Lambda([this]{ Action return FReply::Handled(); }).ToolTipText(LOCTEXT(Category"Tip",Tooltip))]
+#define BUTTON_ANY(Category, Context, Tooltip, Action) +SHorizontalBox::Slot().Padding(StandardPadding).FillWidth(0.5f)[SNew(SButton).HAlign(HAlign_Center).Text(LOCTEXT(Category,Context)).OnClicked_Lambda([this]{ Action return FReply::Handled(); }).IsEnabled(this, &SLidarEditorWidget::IsAnySelection).ToolTipText(LOCTEXT(Category"Tip",Tooltip))]
+#define BUTTON_POINTS(Category, Context, Tooltip, Action) +SHorizontalBox::Slot().Padding(StandardPadding).FillWidth(0.5f)[SNew(SButton).HAlign(HAlign_Center).Text(LOCTEXT(Category,Context)).OnClicked_Lambda([this]{ Action return FReply::Handled(); }).IsEnabled(this, &SLidarEditorWidget::IsPointSelection).ToolTipText(LOCTEXT(Category"Tip",Tooltip))]
+#define BUTTON_ACTORS(Category, Context, Tooltip, Action) +SHorizontalBox::Slot().Padding(StandardPadding).FillWidth(0.5f)[SNew(SButton).HAlign(HAlign_Center).Text(LOCTEXT(Category,Context)).OnClicked_Lambda([this]{ Action return FReply::Handled(); }).IsEnabled(this, &SLidarEditorWidget::IsActorSelection).ToolTipText(LOCTEXT(Category"Tip",Tooltip))]
 	
 	ChildSlot
 	[
@@ -153,7 +152,7 @@ void SLidarEditorWidget::Construct(const FArguments& InArgs)
 		.AutoHeight()
 		[
 			SNew(SHorizontalBox)
-			.ToolTipText(LOCTEXT("CollisionsError", "Determines the maximum error (in cm) of the collision. Lower values will require more time to build."))
+			.ToolTipText(LOCTEXT("CollisionsErrorTooltip", "Determines the maximum error (in cm) of the collision. Lower values will require more time to build."))
 			+ SHorizontalBox::Slot()
 			.Padding(StandardPadding)
 			.FillWidth(0.5f)
@@ -194,7 +193,7 @@ void SLidarEditorWidget::Construct(const FArguments& InArgs)
 		.AutoHeight()
 		[
 			SNew(SHorizontalBox)
-			.ToolTipText(LOCTEXT("NormalsQuality", "Higher values will generally result in more accurate calculations, at the expense of time"))
+			.ToolTipText(LOCTEXT("NormalsQualityTooltip", "Higher values will generally result in more accurate calculations, at the expense of time"))
 			+ SHorizontalBox::Slot()
 			.Padding(StandardPadding)
 			.FillWidth(0.5f)
@@ -220,7 +219,7 @@ void SLidarEditorWidget::Construct(const FArguments& InArgs)
 		.AutoHeight()
 		[
 			SNew(SHorizontalBox)
-			.ToolTipText(LOCTEXT("NormalsNoiseTolerance", "Higher values are less susceptible to noise, but will most likely lose finer details, especially around hard edges."))
+			.ToolTipText(LOCTEXT("NormalsNoiseToleranceTooltip", "Higher values are less susceptible to noise, but will most likely lose finer details, especially around hard edges."))
 			+ SHorizontalBox::Slot()
 			.Padding(StandardPadding)
 			.FillWidth(0.5f)
@@ -266,7 +265,7 @@ void SLidarEditorWidget::Construct(const FArguments& InArgs)
 		.AutoHeight()
 		[
 			SNew(SHorizontalBox)
-			.ToolTipText(LOCTEXT("MeshingError", "Determines the maximum error (in cm) of the resulting mesh. Lower values will require more time to build."))
+			.ToolTipText(LOCTEXT("MeshingErrorTooltip", "Determines the maximum error (in cm) of the resulting mesh. Lower values will require more time to build."))
 			+ SHorizontalBox::Slot()
 			.Padding(StandardPadding)
 			.FillWidth(0.5f)
@@ -294,7 +293,7 @@ void SLidarEditorWidget::Construct(const FArguments& InArgs)
 		.AutoHeight()
 		[
 			SNew(SHorizontalBox)
-			.ToolTipText(LOCTEXT("MeshingMerge", "When enabled, all elements will be combined into one big mesh"))
+			.ToolTipText(LOCTEXT("MeshingMergeTooltip", "When enabled, all elements will be combined into one big mesh"))
 			+ SHorizontalBox::Slot()
 			.Padding(StandardPadding)
 			.FillWidth(0.5f)
@@ -316,7 +315,7 @@ void SLidarEditorWidget::Construct(const FArguments& InArgs)
 		.AutoHeight()
 		[
 			SNew(SHorizontalBox)
-			.ToolTipText(LOCTEXT("MeshingRetain", "When enabled, all generated meshes will retain the transforms of their source lidar data"))
+			.ToolTipText(LOCTEXT("MeshingRetainTooltip", "When enabled, all generated meshes will retain the transforms of their source lidar data"))
 			+ SHorizontalBox::Slot()
 			.Padding(StandardPadding)
 			.FillWidth(0.5f)
