Lidar - Asset : Save as will crash the editor
Viewport blanks and crash after "content changed message"
Autosave fails to reload package when asset had some modifications

diff --git a/Engine/Plugins/Enterprise/LidarPointCloud/Source/LidarPointCloudRuntime/Private/LidarPointCloudOctree.cpp b/Engine/Plugins/Enterprise/LidarPointCloud/Source/LidarPointCloudRuntime/Private/LidarPointCloudOctree.cpp
--- a/Engine/Plugins/Enterprise/LidarPointCloud/Source/LidarPointCloudRuntime/Private/LidarPointCloudOctree.cpp
+++ b/Engine/Plugins/Enterprise/LidarPointCloud/Source/LidarPointCloudRuntime/Private/LidarPointCloudOctree.cpp
@@ -871,7 +871,7 @@ void FLidarPointCloudOctreeNode::AddPointCount(int32 PointCount)
 
 void FLidarPointCloudOctreeNode::SortVisiblePoints()
 {
-	TArrayView<FLidarPointCloudPoint> Points(GetData(), GetNumPoints());
+	TArrayView<FLidarPointCloudPoint> Points(GetData(), Data.Num());
 	Algo::Sort(Points, [](const FLidarPointCloudPoint& A, const FLidarPointCloudPoint& B)
 	{
 		return A.bVisible > B.bVisible;
@@ -2743,7 +2743,26 @@ void FLidarPointCloudOctree::Serialize(FArchive& Ar)
 		}, true);
 		
 #if WITH_EDITOR
-		if(Ar.IsSaving())
+		if(Ar.HasAnyPortFlags(PPF_Duplicate))
+		{
+			if(Ar.IsLoading())
+			{
+				ITERATE_NODES({
+					CurrentNode->Data.SetNumUninitialized(CurrentNode->NumPoints);
+					Ar.Serialize(CurrentNode->Data.GetData(), CurrentNode->BulkDataSize);
+					CurrentNode->bCanReleaseData = false;
+					CurrentNode->bRenderDataDirty = true;
+					CurrentNode->bHasData = true;
+				}, true);	
+			}
+			else
+			{
+				ITERATE_NODES({
+					Ar.Serialize(CurrentNode->Data.GetData(), CurrentNode->BulkDataSize);
+				}, true);	
+			}
+		}
+		else if(Ar.IsSaving())
 		{
 			SavingBulkData.Serialize(Ar, Owner, false, 1, EFileRegionType::None);
 		}
