Lidar - Autosave fails to reload package when asset had some modifications

diff --git a/Engine/Plugins/Enterprise/LidarPointCloud/Source/LidarPointCloudRuntime/Private/LidarPointCloud.cpp b/Engine/Plugins/Enterprise/LidarPointCloud/Source/LidarPointCloudRuntime/Private/LidarPointCloud.cpp
--- a/Engine/Plugins/Enterprise/LidarPointCloud/Source/LidarPointCloudRuntime/Private/LidarPointCloud.cpp
+++ b/Engine/Plugins/Enterprise/LidarPointCloud/Source/LidarPointCloudRuntime/Private/LidarPointCloud.cpp
@@ -77,9 +77,12 @@ private:
 	{
 		if (GetDefault<ULidarPointCloudSettings>()->bReleaseAssetAfterSaving)
 		{
-			if (Cast<ULidarPointCloud>(Package->FindAssetInPackage()))
+			if(Package->GetLinker())
 			{
-				PackagesToReload.Add(Package);
+				if (Cast<ULidarPointCloud>(Package->FindAssetInPackage()))
+				{
+					PackagesToReload.Add(Package);
+				}
 			}
 		}
 	}
diff --git a/Engine/Plugins/Enterprise/LidarPointCloud/Source/LidarPointCloudRuntime/Private/LidarPointCloudOctree.cpp b/Engine/Plugins/Enterprise/LidarPointCloud/Source/LidarPointCloudRuntime/Private/LidarPointCloudOctree.cpp
--- a/Engine/Plugins/Enterprise/LidarPointCloud/Source/LidarPointCloudRuntime/Private/LidarPointCloudOctree.cpp
+++ b/Engine/Plugins/Enterprise/LidarPointCloud/Source/LidarPointCloudRuntime/Private/LidarPointCloudOctree.cpp
@@ -2805,9 +2805,14 @@ void FLidarPointCloudOctree::SerializeBulkData(FArchive& Ar)
 {
 	if (Ar.IsSaving())
 	{
+		const bool bReleaseData = !!Owner->GetPackage()->GetLinker();
+		
 		ITERATE_NODES({
 			Ar.Serialize(CurrentNode->Data.GetData(), CurrentNode->BulkDataSize);
-			CurrentNode->ReleaseData(true);
+			if(bReleaseData)
+			{
+				CurrentNode->ReleaseData(true);
+			}
 		}, true);
 	}
 }
