Fixed non-pch/unity compile error in lidar code

diff --git a/Engine/Plugins/Enterprise/LidarPointCloud/Source/LidarPointCloudRuntime/Private/LidarPointCloudLODManager.cpp b/Engine/Plugins/Enterprise/LidarPointCloud/Source/LidarPointCloudRuntime/Private/LidarPointCloudLODManager.cpp
--- a/Engine/Plugins/Enterprise/LidarPointCloud/Source/LidarPointCloudRuntime/Private/LidarPointCloudLODManager.cpp
+++ b/Engine/Plugins/Enterprise/LidarPointCloud/Source/LidarPointCloudRuntime/Private/LidarPointCloudLODManager.cpp
@@ -1,7 +1,6 @@
 // Copyright Epic Games, Inc. All Rights Reserved.
 
 #include "LidarPointCloudLODManager.h"
-#include "Components/Viewport.h"
 #include "LidarPointCloud.h"
 #include "LidarPointCloudComponent.h"
 #include "LidarPointCloudOctree.h"
@@ -14,6 +13,7 @@
 #include "Misc/ScopeTryLock.h"
 #include "Engine/Engine.h"
 #include "EngineUtils.h"
+#include "UnrealClient.h"
 
 #if WITH_EDITOR
 #include "Settings/EditorStyleSettings.h"
