Lidar meshing now uses MeshDescription.
Fixed a few crashes related to mesh render data.

diff --git a/Engine/Plugins/Enterprise/LidarPointCloud/Source/LidarPointCloudEditor/LidarPointCloudEditor.Build.cs b/Engine/Plugins/Enterprise/LidarPointCloud/Source/LidarPointCloudEditor/LidarPointCloudEditor.Build.cs
--- a/Engine/Plugins/Enterprise/LidarPointCloud/Source/LidarPointCloudEditor/LidarPointCloudEditor.Build.cs
+++ b/Engine/Plugins/Enterprise/LidarPointCloud/Source/LidarPointCloudEditor/LidarPointCloudEditor.Build.cs
@@ -35,6 +35,7 @@ namespace UnrealBuildTool.Rules
                 "InteractiveToolsFramework",
                 "StatusBar",
                 "MeshDescription",
+                "StaticMeshDescription"
                 }
             );
         }
diff --git a/Engine/Plugins/Enterprise/LidarPointCloud/Source/LidarPointCloudEditor/Private/LidarPointCloudEditorHelper.cpp b/Engine/Plugins/Enterprise/LidarPointCloud/Source/LidarPointCloudEditor/Private/LidarPointCloudEditorHelper.cpp
--- a/Engine/Plugins/Enterprise/LidarPointCloud/Source/LidarPointCloudEditor/Private/LidarPointCloudEditorHelper.cpp
+++ b/Engine/Plugins/Enterprise/LidarPointCloud/Source/LidarPointCloudEditor/Private/LidarPointCloudEditorHelper.cpp
@@ -10,6 +10,7 @@
 #include "LidarPointCloudComponent.h"
 #include "MeshDescription.h"
 #include "Selection.h"
+#include "StaticMeshAttributes.h"
 #include "AssetRegistry/AssetRegistryModule.h"
 #include "Engine/StaticMeshActor.h"
 #include "Misc/ScopedSlowTask.h"
@@ -203,8 +204,84 @@ namespace
 		return StaticMesh;
 	}
 
+	void MeshBuffersToMeshDescription(LidarPointCloudMeshing::FMeshBuffers& MeshBuffers, FMeshDescription& OutMeshDescription)
+	{
+		FScopeBenchmarkTimer Timer("MeshBuffersToMeshDescription");
+		
+		FStaticMeshAttributes Attributes(OutMeshDescription);
+		Attributes.Register();
+		TVertexAttributesRef<FVector3f> VertexPositions = Attributes.GetVertexPositions();
+		TVertexInstanceAttributesRef<FVector3f> InstanceNormals = Attributes.GetVertexInstanceNormals();
+		TVertexInstanceAttributesRef<FVector4f> InstanceColors = Attributes.GetVertexInstanceColors();
+
+		OutMeshDescription.SuspendVertexInstanceIndexing();
+		OutMeshDescription.SuspendEdgeIndexing();
+		OutMeshDescription.SuspendPolygonIndexing();
+		OutMeshDescription.SuspendPolygonGroupIndexing();
+		OutMeshDescription.SuspendUVIndexing();
+		OutMeshDescription.ReserveNewVertices(MeshBuffers.Vertices.Num());
+		OutMeshDescription.ReserveNewVertexInstances(MeshBuffers.Indices.Num());
+
+		const FPolygonGroupID AllGroupID = OutMeshDescription.CreatePolygonGroup();
+		Attributes.GetPolygonGroupMaterialSlotNames().Set(AllGroupID, NAME_None);
+
+		for(LidarPointCloudMeshing::FVertexData& Vertex : MeshBuffers.Vertices)
+		{
+			VertexPositions.Set(OutMeshDescription.CreateVertex(), Vertex.Position);
+		}
+
+		TArray<FVertexInstanceID> CornerInstanceIDs;
+		CornerInstanceIDs.SetNumZeroed(3);
+		for (uint32* Index = MeshBuffers.Indices.GetData(), *DataEnd = Index + MeshBuffers.Indices.Num(); Index != DataEnd;)
+		{
+			for(int32 i = 0; i < 3; ++i, ++Index)
+			{
+				FVertexInstanceID& VertexInstanceID = CornerInstanceIDs[i];
+				VertexInstanceID = OutMeshDescription.CreateVertexInstance(*Index);
+				LidarPointCloudMeshing::FVertexData& VertexData = MeshBuffers.Vertices[*Index];
+				InstanceNormals.Set(VertexInstanceID, VertexData.Normal);
+				InstanceColors.Set(VertexInstanceID, FLinearColor(VertexData.Color));
+			}
+				
+			OutMeshDescription.CreateTriangle(AllGroupID, CornerInstanceIDs);
+		}
+		
+		OutMeshDescription.ResumeVertexInstanceIndexing();
+		OutMeshDescription.ResumeEdgeIndexing();
+		OutMeshDescription.ResumePolygonIndexing();
+		OutMeshDescription.ResumePolygonGroupIndexing();
+		OutMeshDescription.ResumeUVIndexing();
+	}
+
+	void AssignMeshDescriptionToMesh(FMeshDescription& MeshDescription, UStaticMesh* StaticMesh)
+	{
+		FScopeBenchmarkTimer Timer("AssignMeshDescriptionToMesh");
+		
+		if(StaticMesh)
+		{
+			StaticMesh->SetNumSourceModels(1);
+			FStaticMeshSourceModel& SourceModelLOD0 = StaticMesh->GetSourceModel(0);
+			SourceModelLOD0.BuildSettings.bRecomputeNormals = false;
+			SourceModelLOD0.BuildSettings.bRecomputeTangents = false;
+			SourceModelLOD0.BuildSettings.bBuildReversedIndexBuffer = false;
+			SourceModelLOD0.BuildSettings.bGenerateLightmapUVs = false;
+
+			TArray<FStaticMaterial> StaticMaterials;
+			static UMaterialInterface* Material = Cast<UMaterialInterface>(FSoftObjectPath(TEXT("/LidarPointCloud/Materials/M_MeshedCloud.M_MeshedCloud")).TryLoad());
+			StaticMaterials.Emplace(Material);
+			StaticMesh->SetStaticMaterials(StaticMaterials);
+			StaticMesh->BuildFromMeshDescriptions({ &MeshDescription });
+			
+			StaticMesh->SetLightingGuid(FGuid::NewGuid());
+			StaticMesh->ImportVersion = EImportStaticMeshVersion::LastVersion;
+			StaticMesh->PostEditChange();
+		}
+	}
+
 	void AssignMeshBuffersToMesh(LidarPointCloudMeshing::FMeshBuffers* MeshBuffers, UStaticMesh* StaticMesh)
 	{
+		FScopeBenchmarkTimer Timer("AssignMeshBuffersToMesh");
+		
 		if(StaticMesh && MeshBuffers)
 		{
 			constexpr bool bCPUAccess = true;
@@ -286,8 +363,10 @@ namespace
 				}
 			}
 			
+			TArray<FStaticMaterial> StaticMaterials;
 			static UMaterialInterface* Material = Cast<UMaterialInterface>(FSoftObjectPath(TEXT("/LidarPointCloud/Materials/M_MeshedCloud.M_MeshedCloud")).TryLoad());
-			StaticMesh->AddMaterial(Material);
+			StaticMaterials.Emplace(Material);
+			StaticMesh->SetStaticMaterials(StaticMaterials);
 
 			StaticMesh->SetLightingGuid(FGuid::NewGuid()) ;
 			StaticMesh->ImportVersion = EImportStaticMeshVersion::LastVersion;
@@ -536,8 +615,11 @@ void FLidarPointCloudEditorHelper::MeshSelected(bool bMeshByPoints, float CellSi
 			{
 				ProcessSelection(MoveTemp(MergeFunc));
 			}
-		
-			AssignMeshBuffersToMesh(&MeshBuffers, StaticMesh);
+
+			FMeshDescription MeshDescription;
+			MeshBuffersToMeshDescription(MeshBuffers, MeshDescription);
+			AssignMeshDescriptionToMesh(MeshDescription, StaticMesh);
+
 			ProgressDialog.EnterProgressFrame(1.0f);
 					
 			if(AStaticMeshActor* MeshActor = SpawnActor<AStaticMeshActor>(StaticMesh->GetName()))
@@ -565,7 +647,9 @@ void FLidarPointCloudEditorHelper::MeshSelected(bool bMeshByPoints, float CellSi
 						PointCloud->BuildStaticMeshBuffers(CellSize, &MeshBuffers, bRetainTransform ? FTransform::Identity : Actor->GetActorTransform());
 					}
 					
-					AssignMeshBuffersToMesh(&MeshBuffers, StaticMesh);
+					FMeshDescription MeshDescription;
+					MeshBuffersToMeshDescription(MeshBuffers, MeshDescription);
+					AssignMeshDescriptionToMesh(MeshDescription, StaticMesh);
 					
 					if(AStaticMeshActor* MeshActor = SpawnActor<AStaticMeshActor>(StaticMesh->GetName()))
 					{
