Fix Lidar point cloud shader compilation issues in ISR.

diff --git a/Engine/Plugins/Enterprise/LidarPointCloud/Shaders/Private/LidarPointCloudVertexFactory.ush b/Engine/Plugins/Enterprise/LidarPointCloud/Shaders/Private/LidarPointCloudVertexFactory.ush
--- a/Engine/Plugins/Enterprise/LidarPointCloud/Shaders/Private/LidarPointCloudVertexFactory.ush
+++ b/Engine/Plugins/Enterprise/LidarPointCloud/Shaders/Private/LidarPointCloudVertexFactory.ush
@@ -69,25 +69,22 @@ float4x4 ClippingVolume[16];
 uint NumClippingVolumes;
 uint bStartClipped;
 
-#if USE_INSTANCING
-	// Instanced stereo sets the eye index explicitly for instanced geometry
-	#if INSTANCED_STEREO
-		uint InstancedEyeIndex;
-	#endif
-#endif
-
 struct FVertexFactoryInput
 {
     float4	Position    : ATTRIBUTE0;
     half4	Color		: ATTRIBUTE1;
     uint	MetaData	: ATTRIBUTE2;
 	uint    VertexId    : SV_VertexID;
+
+	VF_INSTANCED_STEREO_DECLARE_INPUT_BLOCK()
 };
 
 struct FPositionOnlyVertexFactoryInput
 {
 	float4  Position    : ATTRIBUTE0;
 	uint    VertexId    : SV_VertexID;
+
+	VF_INSTANCED_STEREO_DECLARE_INPUT_BLOCK()
 };
 
 struct FPositionAndNormalOnlyVertexFactoryInput
@@ -95,6 +92,8 @@ struct FPositionAndNormalOnlyVertexFactoryInput
 	float4  Position	: ATTRIBUTE0;
 	float4	Normal		: ATTRIBUTE2;
 	uint    VertexId	: SV_VertexID;
+
+	VF_INSTANCED_STEREO_DECLARE_INPUT_BLOCK()
 };
 
 struct FVertexFactoryInterpolantsVSToPS
@@ -375,7 +374,7 @@ FVertexFactoryInterpolantsVSToPS VertexFactoryGetInterpolantsVSToPS(FVertexFacto
 	Interpolants.TexCoords = PrecomputedUV_Quad[Input.VertexId % 4 + 4];    // Using inverted V
 
 #if INSTANCED_STEREO
-	Interpolants.EyeIndex = 0;
+	Interpolants.EyeIndex = GetEyeIndexFromVF(Input);
 #endif
 
 	return Interpolants;
@@ -477,17 +476,6 @@ float4 VertexFactoryGetTranslatedPrimitiveVolumeBounds(FVertexFactoryInterpolant
 	return float4(LWCToFloat(LWCAdd(PrimitiveData.ObjectWorldPosition, ResolvedView.PreViewTranslation)), PrimitiveData.ObjectRadius);
 }
 
-#if INSTANCED_STEREO
-uint VertexFactoryGetEyeIndex(uint InstanceId)
-{
-	#if USE_INSTANCING
-		return InstancedEyeIndex;
-	#else
-		return InstanceId & 1;
-	#endif
-}
-#endif
-
 #if NEEDS_VERTEX_FACTORY_INTERPOLATION
 	struct FVertexFactoryRayTracingInterpolants
 	{
