LidarPointCloud - Replace usage of IAssetTypeActions_AssetInterface with UAssetDefinition

diff --git a/Engine/Plugins/Enterprise/LidarPointCloud/Source/LidarPointCloudEditor/LidarPointCloudEditor.Build.cs b/Engine/Plugins/Enterprise/LidarPointCloud/Source/LidarPointCloudEditor/LidarPointCloudEditor.Build.cs
--- a/Engine/Plugins/Enterprise/LidarPointCloud/Source/LidarPointCloudEditor/LidarPointCloudEditor.Build.cs
+++ b/Engine/Plugins/Enterprise/LidarPointCloud/Source/LidarPointCloudEditor/LidarPointCloudEditor.Build.cs
@@ -11,7 +11,10 @@ namespace UnrealBuildTool.Rules
             PrivateDependencyModuleNames.AddRange(
                 new string[]
                 {
+                "AssetDefinition",
                 "CoreUObject",
+                "ContentBrowser",
+				"ToolMenus",
                 "Engine",
                 "Slate",
                 "SlateCore",
diff --git a/Engine/Plugins/Enterprise/LidarPointCloud/Source/LidarPointCloudEditor/Private/AssetDefinition_LidarPointCloud.cpp b/Engine/Plugins/Enterprise/LidarPointCloud/Source/LidarPointCloudEditor/Private/AssetDefinition_LidarPointCloud.cpp
new file mode 100644
--- /dev/null
+++ b/Engine/Plugins/Enterprise/LidarPointCloud/Source/LidarPointCloudEditor/Private/AssetDefinition_LidarPointCloud.cpp
@@ -0,0 +1,172 @@
+// Copyright Epic Games, Inc. All Rights Reserved.
+
+#include "AssetDefinition_LidarPointCloud.h"
+
+#include "ContentBrowserMenuContexts.h"
+#include "PackageTools.h"
+
+#include "LidarPointCloudEditor.h"
+#include "LidarPointCloudEditorHelper.h"
+
+#include "LidarPointCloud.h"
+
+
+#define LOCTEXT_NAMESPACE "LidarPointCloud"
+
+
+FText UAssetDefinition_LidarPointCloud::GetAssetDisplayName() const
+{
+	return NSLOCTEXT("AssetTypeActions", "FAssetTypeActions_LidarPointCloud", "LiDAR Point Cloud");
+}
+
+EAssetCommandResult UAssetDefinition_LidarPointCloud::OpenAssets(const FAssetOpenArgs& OpenArgs) const
+{
+	for (ULidarPointCloud* PointCloud : OpenArgs.LoadObjects<ULidarPointCloud>())
+	{
+		TSharedRef<FLidarPointCloudEditor> NewPointCloudEditor(new FLidarPointCloudEditor());
+		NewPointCloudEditor->InitPointCloudEditor(EToolkitMode::Standalone, OpenArgs.ToolkitHost, PointCloud);
+	}
+
+	return EAssetCommandResult::Handled;
+}
+
+TConstArrayView<FAssetCategoryPath> UAssetDefinition_LidarPointCloud::GetAssetCategories() const
+{
+	static const auto Categories = { EAssetCategoryPaths::Misc };
+	return Categories;
+}
+
+TSoftClassPtr<UObject> UAssetDefinition_LidarPointCloud::GetAssetClass() const
+{
+	return ULidarPointCloud::StaticClass();
+}
+
+namespace MenuExtension_LidarPointCloud
+{
+
+	void ExecuteMerge(TArray<ULidarPointCloud*> PointClouds)
+	{
+		if (PointClouds.Num() < 2)
+		{
+			return;
+		}
+
+		FString PackageName = PointClouds[0]->GetOutermost()->GetName() + TEXT("_Merged");
+		UPackage* MergedCloudPackage = UPackageTools::FindOrCreatePackageForAssetType(FName(*PackageName), ULidarPointCloud::StaticClass());
+		if (IsValid(MergedCloudPackage))
+		{
+			MergedCloudPackage->SetPackageFlags(PKG_NewlyCreated);
+
+			ULidarPointCloud* PC = NewObject<ULidarPointCloud>(MergedCloudPackage, FName(*FPaths::GetBaseFilename(MergedCloudPackage->GetName())), EObjectFlags::RF_Public | EObjectFlags::RF_Standalone | EObjectFlags::RF_Transactional);
+
+			FLidarPointCloudEditorHelper::MergeLidar(PC, PointClouds);
+		}
+	}
+
+	void ExecuteAlign(TArray<ULidarPointCloud*> PointClouds)
+	{
+		if (PointClouds.Num() < 2)
+		{
+			return;
+		}
+
+		FScopedSlowTask ProgressDialog(1, LOCTEXT("Align", "Aligning Point Clouds..."));
+		ProgressDialog.MakeDialog();
+		ProgressDialog.EnterProgressFrame(1.f);
+		ULidarPointCloud::AlignClouds(PointClouds);
+	}
+
+	void ExecuteCollision(TArray<ULidarPointCloud*> PointClouds)
+	{
+		for (ULidarPointCloud* PC : PointClouds)
+		{
+			PC->BuildCollision();
+		}
+	}
+
+	void ExecuteNormals(TArray<ULidarPointCloud*> PointClouds)
+	{
+		for (ULidarPointCloud* PC : PointClouds)
+		{
+			// Data needs to be persistently loaded to calculate normals
+			if (!PC->IsFullyLoaded())
+			{
+				PC->LoadAllNodes();
+			}
+
+			PC->CalculateNormals(nullptr, nullptr);
+		}
+	}
+
+
+
+	static FDelayedAutoRegisterHelper DelayedAutoRegister(EDelayedRegisterRunPhase::EndOfEngineInit, []{ 
+		UToolMenus::RegisterStartupCallback(FSimpleMulticastDelegate::FDelegate::CreateLambda([]()
+		{
+			FToolMenuOwnerScoped OwnerScoped(UE_MODULE_NAME);
+			UToolMenu* Menu = UE::ContentBrowser::ExtendToolMenu_AssetContextMenu(ULidarPointCloud::StaticClass());
+	        
+			FToolMenuSection& Section = Menu->FindOrAddSection("GetAssetActions");
+			Section.AddDynamicEntry(NAME_None, FNewToolMenuSectionDelegate::CreateLambda([](FToolMenuSection& InSection)
+			{
+				if (const UContentBrowserAssetContextMenuContext* CBContext = UContentBrowserAssetContextMenuContext::FindContextWithAssets(InSection))
+				{
+					TArray<FAssetData> Assets = CBContext->GetSelectedAssetsOfType(ULidarPointCloud::StaticClass());
+
+					if (Assets.Num() > 1)
+					{
+						InSection.AddMenuEntry("LidarPointCloud_Merge", 
+							LOCTEXT("LidarPointCloud_Merge", "Merge Selected"), 
+							LOCTEXT("LidarPointCloud_MergeTooltip", "Merges selected point cloud assets."), 
+							FSlateIcon("LidarPointCloudStyle", "LidarPointCloudEditor.ToolkitMerge"),
+							FToolMenuExecuteAction::CreateLambda([](const FToolMenuContext& Context)
+							{
+								if (const UContentBrowserAssetContextMenuContext* CBContext = UContentBrowserAssetContextMenuContext::FindContextWithAssets(Context))
+								{
+									ExecuteMerge(CBContext->LoadSelectedObjects<ULidarPointCloud>());
+								}
+							}));
+
+						InSection.AddMenuEntry("LidarPointCloud_Align",
+							LOCTEXT("LidarPointCloud_Align", "Align Selected"),
+							LOCTEXT("LidarPointCloud_AlignTooltip", "Aligns selected point cloud assets."),
+							FSlateIcon("LidarPointCloudStyle", "LidarPointCloudEditor.ToolkitAlign"),
+							FToolMenuExecuteAction::CreateLambda([](const FToolMenuContext& Context)
+							{
+								if (const UContentBrowserAssetContextMenuContext* CBContext = UContentBrowserAssetContextMenuContext::FindContextWithAssets(Context))
+								{
+									ExecuteAlign(CBContext->LoadSelectedObjects<ULidarPointCloud>());
+								}
+							}));
+					}
+
+					InSection.AddMenuEntry("LidarPointCloud_BuildCollision",
+						LOCTEXT("LidarPointCloud_BuildCollision", "Build Collision"),
+						LOCTEXT("LidarPointCloud_BuildCollisionTooltip", "Builds collision for all selected point cloud assets."),
+						FSlateIcon("LidarPointCloudStyle", "LidarPointCloudEditor.ToolkitCollision"),
+						FToolMenuExecuteAction::CreateLambda([](const FToolMenuContext& Context)
+						{
+							if (const UContentBrowserAssetContextMenuContext* CBContext = UContentBrowserAssetContextMenuContext::FindContextWithAssets(Context))
+							{
+								ExecuteCollision(CBContext->LoadSelectedObjects<ULidarPointCloud>());
+							}
+						}));
+
+					InSection.AddMenuEntry("LidarPointCloud_CalculateNormals",
+						LOCTEXT("LidarPointCloud_CalculateNormals", "Calculate Normals"),
+						LOCTEXT("LidarPointCloud_CalculateNormalsTooltip", "Calculates normals for all selected point cloud assets."),
+						FSlateIcon("LidarPointCloudStyle", "LidarPointCloudEditor.ToolkitNormals"),
+						FToolMenuExecuteAction::CreateLambda([](const FToolMenuContext& Context)
+						{
+							if (const UContentBrowserAssetContextMenuContext* CBContext = UContentBrowserAssetContextMenuContext::FindContextWithAssets(Context))
+							{
+								ExecuteNormals(CBContext->LoadSelectedObjects<ULidarPointCloud>());
+							}
+						}));
+				}
+			}));
+		}));
+	});
+}
+
+#undef LOCTEXT_NAMESPACE
\ No newline at end of file
diff --git a/Engine/Plugins/Enterprise/LidarPointCloud/Source/LidarPointCloudEditor/Private/AssetDefinition_LidarPointCloud.h b/Engine/Plugins/Enterprise/LidarPointCloud/Source/LidarPointCloudEditor/Private/AssetDefinition_LidarPointCloud.h
new file mode 100644
--- /dev/null
+++ b/Engine/Plugins/Enterprise/LidarPointCloud/Source/LidarPointCloudEditor/Private/AssetDefinition_LidarPointCloud.h
@@ -0,0 +1,22 @@
+// Copyright Epic Games, Inc. All Rights Reserved.
+
+#pragma once
+#include "AssetDefinitionDefault.h"
+#include "AssetDefinition_LidarPointCloud.generated.h"
+
+
+UCLASS()
+class UAssetDefinition_LidarPointCloud : public UAssetDefinitionDefault
+{
+	GENERATED_BODY()
+
+public:
+	// UAssetDefinition Begin
+	virtual FText GetAssetDisplayName() const override;
+	virtual FLinearColor GetAssetColor() const override { return FColor(0, 128, 128); }
+	virtual TSoftClassPtr<UObject> GetAssetClass() const override;
+	virtual EAssetCommandResult OpenAssets(const FAssetOpenArgs& OpenArgs) const override;
+	virtual TConstArrayView<FAssetCategoryPath> GetAssetCategories() const override;
+	virtual bool CanImport() const override { return true; }
+	// UAssetDefinition End
+};
diff --git a/Engine/Plugins/Enterprise/LidarPointCloud/Source/LidarPointCloudEditor/Private/LidarPointCloudEditorModule.cpp b/Engine/Plugins/Enterprise/LidarPointCloud/Source/LidarPointCloudEditor/Private/LidarPointCloudEditorModule.cpp
--- a/Engine/Plugins/Enterprise/LidarPointCloud/Source/LidarPointCloudEditor/Private/LidarPointCloudEditorModule.cpp
+++ b/Engine/Plugins/Enterprise/LidarPointCloud/Source/LidarPointCloudEditor/Private/LidarPointCloudEditorModule.cpp
@@ -33,10 +33,6 @@ class FLidarPointCloudEditorModule : public ILidarPointCloudEditorModule
 		// Register slate style overrides
 		FLidarPointCloudStyle::Initialize();
 		
-		// Register asset type
-		IAssetTools& AssetTools = FModuleManager::LoadModuleChecked<FAssetToolsModule>("AssetTools").Get();
-		AssetTools.RegisterAssetTypeActions(MakeShareable(new FAssetTypeActions_LidarPointCloud));
-
 		// Register settings
 		if (ISettingsModule * SettingsModule = FModuleManager::GetModulePtr<ISettingsModule>("Settings"))
 		{
@@ -68,13 +64,6 @@ class FLidarPointCloudEditorModule : public ILidarPointCloudEditorModule
 			}
 		}
 
-		// Unregister the asset type that we registered
-		if (FModuleManager::Get().IsModuleLoaded("AssetTools"))
-		{
-			IAssetTools& AssetTools = FModuleManager::GetModuleChecked<FAssetToolsModule>("AssetTools").Get();
-			AssetTools.UnregisterAssetTypeActions(MakeShareable(new FAssetTypeActions_LidarPointCloud));
-		}
-		
 		// Unregister slate style overrides
 		FLidarPointCloudStyle::Shutdown();
 	}
diff --git a/Engine/Plugins/Enterprise/LidarPointCloud/Source/LidarPointCloudEditor/Private/LidarPointCloudFactory.cpp b/Engine/Plugins/Enterprise/LidarPointCloud/Source/LidarPointCloudEditor/Private/LidarPointCloudFactory.cpp
--- a/Engine/Plugins/Enterprise/LidarPointCloud/Source/LidarPointCloudEditor/Private/LidarPointCloudFactory.cpp
+++ b/Engine/Plugins/Enterprise/LidarPointCloud/Source/LidarPointCloudEditor/Private/LidarPointCloudFactory.cpp
@@ -1,152 +1,20 @@
 // Copyright Epic Games, Inc. All Rights Reserved.
 
 #include "LidarPointCloudFactory.h"
+
+#include "ContentBrowserMenuContexts.h"
 #include "LidarPointCloudShared.h"
 #include "LidarPointCloudSettings.h"
 #include "LidarPointCloudImportUI.h"
-#include "LidarPointCloudEditor.h"
 #include "IO/LidarPointCloudFileIO.h"
-#include "Styling/AppStyle.h"
-#include "Framework/MultiBox/MultiBoxBuilder.h"
 #include "Editor.h"
-#include "LidarPointCloudEditorHelper.h"
 
-#include "Misc/ScopedSlowTask.h"
 #include "AssetRegistry/AssetRegistryModule.h"
-#include "PackageTools.h"
 
 #define LOCTEXT_NAMESPACE "LidarPointCloud"
 
-FText FAssetTypeActions_LidarPointCloud::GetName() const
-{
-	return NSLOCTEXT("AssetTypeActions", "FAssetTypeActions_LidarPointCloud", "LiDAR Point Cloud");
-}
-
-void FAssetTypeActions_LidarPointCloud::GetActions(const TArray<UObject*>& InObjects, FMenuBuilder& MenuBuilder)
-{
-	TArray<ULidarPointCloud*> PointClouds;
-	for (UObject* Object : InObjects)
-	{
-		PointClouds.Add(CastChecked<ULidarPointCloud>(Object));
-	}
-
-	if (PointClouds.Num() > 1)
-	{
-		MenuBuilder.AddMenuEntry(
-			LOCTEXT("LidarPointCloud_Merge", "Merge Selected"),
-			LOCTEXT("LidarPointCloud_MergeTooltip", "Merges selected point cloud assets."),
-			FSlateIcon("LidarPointCloudStyle", "LidarPointCloudEditor.ToolkitMerge"),
-			FUIAction(
-				FExecuteAction::CreateSP(this, &FAssetTypeActions_LidarPointCloud::ExecuteMerge, PointClouds),
-				FCanExecuteAction()
-			)
-		);
-
-		MenuBuilder.AddMenuEntry(
-			LOCTEXT("LidarPointCloud_Align", "Align Selected"),
-			LOCTEXT("LidarPointCloud_AlignTooltip", "Aligns selected point cloud assets."),
-			FSlateIcon("LidarPointCloudStyle", "LidarPointCloudEditor.ToolkitAlign"),
-			FUIAction(
-				FExecuteAction::CreateSP(this, &FAssetTypeActions_LidarPointCloud::ExecuteAlign, PointClouds),
-				FCanExecuteAction()
-			)
-		);
-	}
 
-	MenuBuilder.AddMenuEntry(
-		LOCTEXT("LidarPointCloud_BuildCollision", "Build Collision"),
-		LOCTEXT("LidarPointCloud_BuildCollisionTooltip", "Builds collision for all selected point cloud assets."),
-		FSlateIcon("LidarPointCloudStyle", "LidarPointCloudEditor.ToolkitCollision"),
-		FUIAction(
-			FExecuteAction::CreateSP(this, &FAssetTypeActions_LidarPointCloud::ExecuteCollision, PointClouds),
-			FCanExecuteAction()
-		)
-	);
 
-	MenuBuilder.AddMenuEntry(
-		LOCTEXT("LidarPointCloud_CalculateNormals", "Calculate Normals"),
-		LOCTEXT("LidarPointCloud_CalculateNormalsTooltip", "Calculates normals for all selected point cloud assets."),
-		FSlateIcon("LidarPointCloudStyle", "LidarPointCloudEditor.ToolkitNormals"),
-		FUIAction(
-			FExecuteAction::CreateSP(this, &FAssetTypeActions_LidarPointCloud::ExecuteNormals, PointClouds),
-			FCanExecuteAction()
-		)
-	);
-}
-
-void FAssetTypeActions_LidarPointCloud::OpenAssetEditor(const TArray<UObject*>& InObjects, TSharedPtr<class IToolkitHost> EditWithinLevelEditor /*= TSharedPtr<IToolkitHost>()*/)
-{
-	const EToolkitMode::Type Mode = EditWithinLevelEditor.IsValid() ? EToolkitMode::WorldCentric : EToolkitMode::Standalone;
-
-	for (UObject* Object : InObjects)
-	{
-		if (ULidarPointCloud* PointCloud = Cast<ULidarPointCloud>(Object))
-		{
-			TSharedRef<FLidarPointCloudEditor> NewPointCloudEditor(new FLidarPointCloudEditor());
-			NewPointCloudEditor->InitPointCloudEditor(Mode, EditWithinLevelEditor, PointCloud);
-		}
-	}
-}
-
-void FAssetTypeActions_LidarPointCloud::GetResolvedSourceFilePaths(const TArray<UObject*>& TypeAssets, TArray<FString>& OutSourceFilePaths) const
-{
-	for (auto& Asset : TypeAssets)
-	{
-		if (const ULidarPointCloud* PointCloud = CastChecked<ULidarPointCloud>(Asset))
-		{
-			OutSourceFilePaths.Add(PointCloud->GetSourcePath());
-		}
-	}
-}
-
-void FAssetTypeActions_LidarPointCloud::ExecuteMerge(TArray<ULidarPointCloud*> PointClouds)
-{
-	if (PointClouds.Num() < 2)
-	{
-		return;
-	}
-
-	FString PackageName = PointClouds[0]->GetOutermost()->GetName() + TEXT("_Merged");
-	UPackage* MergedCloudPackage = UPackageTools::FindOrCreatePackageForAssetType(FName(*PackageName), ULidarPointCloud::StaticClass());
-	if (IsValid(MergedCloudPackage))
-	{
-		MergedCloudPackage->SetPackageFlags(PKG_NewlyCreated);
-
-		ULidarPointCloud* PC = NewObject<ULidarPointCloud>(MergedCloudPackage, FName(*FPaths::GetBaseFilename(MergedCloudPackage->GetName())), EObjectFlags::RF_Public | EObjectFlags::RF_Standalone | EObjectFlags::RF_Transactional);
-
-		FLidarPointCloudEditorHelper::MergeLidar(PC, PointClouds);
-	}
-}
-
-void FAssetTypeActions_LidarPointCloud::ExecuteAlign(TArray<ULidarPointCloud*> PointClouds)
-{
-	FScopedSlowTask ProgressDialog(1, LOCTEXT("Align", "Aligning Point Clouds..."));
-	ProgressDialog.MakeDialog();
-	ProgressDialog.EnterProgressFrame(1.f);
-	ULidarPointCloud::AlignClouds(PointClouds);
-}
-
-void FAssetTypeActions_LidarPointCloud::ExecuteCollision(TArray<ULidarPointCloud*> PointClouds)
-{
-	for (ULidarPointCloud* PC : PointClouds)
-	{
-		PC->BuildCollision();
-	}
-}
-
-void FAssetTypeActions_LidarPointCloud::ExecuteNormals(TArray<ULidarPointCloud*> PointClouds)
-{
-	for (ULidarPointCloud* PC : PointClouds)
-	{
-		// Data needs to be persistently loaded to calculate normals
-		if (!PC->IsFullyLoaded())
-		{
-			PC->LoadAllNodes();
-		}
-
-		PC->CalculateNormals(nullptr, nullptr);
-	}
-}
 
 ULidarPointCloudFactory::ULidarPointCloudFactory()
 {
diff --git a/Engine/Plugins/Enterprise/LidarPointCloud/Source/LidarPointCloudEditor/Private/LidarPointCloudFactory.h b/Engine/Plugins/Enterprise/LidarPointCloud/Source/LidarPointCloudEditor/Private/LidarPointCloudFactory.h
--- a/Engine/Plugins/Enterprise/LidarPointCloud/Source/LidarPointCloudEditor/Private/LidarPointCloudFactory.h
+++ b/Engine/Plugins/Enterprise/LidarPointCloud/Source/LidarPointCloudEditor/Private/LidarPointCloudFactory.h
@@ -2,34 +2,11 @@
 
 #pragma once
 #include "EditorReimportHandler.h"
-#include "AssetTypeCategories.h"
 #include "AssetTypeActions_Base.h"
 #include "Factories/Factory.h"
 #include "LidarPointCloud.h"
 #include "LidarPointCloudFactory.generated.h"
 
-class FAssetTypeActions_LidarPointCloud : public FAssetTypeActions_Base
-{
-public:
-	FAssetTypeActions_LidarPointCloud() {}
-
-	// Begin IAssetTypeActions Interface
-	virtual FText GetName() const override;
-	virtual FColor GetTypeColor() const override { return FColor(0, 128, 128); }
-	virtual UClass* GetSupportedClass() const override { return ULidarPointCloud::StaticClass(); }
-	virtual void GetActions(const TArray<UObject*>& InObjects, FMenuBuilder& MenuBuilder) override;
-	virtual uint32 GetCategories() override { return EAssetTypeCategories::Misc; }
-	virtual bool IsImportedAsset() const override { return true; }
-	virtual void OpenAssetEditor(const TArray<UObject*>& InObjects, TSharedPtr<class IToolkitHost> EditWithinLevelEditor = TSharedPtr<IToolkitHost>()) override;
-	virtual void GetResolvedSourceFilePaths(const TArray<UObject*>& TypeAssets, TArray<FString>& OutSourceFilePaths) const override;
-	// End IAssetTypeActions Interface
-
-private:
-	void ExecuteMerge(TArray<ULidarPointCloud*> PointClouds);
-	void ExecuteAlign(TArray<ULidarPointCloud*> PointClouds);
-	void ExecuteCollision(TArray<ULidarPointCloud*> PointClouds);
-	void ExecuteNormals(TArray<ULidarPointCloud*> PointClouds);
-};
 
 UCLASS()
 class ULidarPointCloudFactory : public UFactory, public FReimportHandler
