Fixed mismatched copy constructor/assignment operator in FLidarPointCloudPoint.

diff --git a/Engine/Plugins/Enterprise/LidarPointCloud/Source/LidarPointCloudRuntime/Public/LidarPointCloudShared.h b/Engine/Plugins/Enterprise/LidarPointCloud/Source/LidarPointCloudRuntime/Public/LidarPointCloudShared.h
--- a/Engine/Plugins/Enterprise/LidarPointCloud/Source/LidarPointCloudRuntime/Public/LidarPointCloudShared.h
+++ b/Engine/Plugins/Enterprise/LidarPointCloud/Source/LidarPointCloudRuntime/Public/LidarPointCloudShared.h
@@ -161,11 +161,13 @@ public:
 		this->ClassificationID = ClassificationID;
 		this->Normal = Normal;
 	}
+	// Questionable behavior here - the copy constructor only copies part of the state but the assignment operator copies it all
 	FLidarPointCloudPoint(const FLidarPointCloudPoint& Other)
 		: FLidarPointCloudPoint()
 	{
 		CopyFrom(Other);
 	}
+	FLidarPointCloudPoint& operator=(const FLidarPointCloudPoint&) = default;
 	FLidarPointCloudPoint(const FLidarPointCloudPoint_Legacy& Other)
 		: FLidarPointCloudPoint(Other.Location, Other.Color, Other.bVisible, Other.ClassificationID)
 	{
