Speculative include fix in lidar code

#rnx

diff --git a/Engine/Plugins/Enterprise/LidarPointCloud/Source/LidarPointCloudRuntime/Public/Meshing/LidarPointCloudMeshing.h b/Engine/Plugins/Enterprise/LidarPointCloud/Source/LidarPointCloudRuntime/Public/Meshing/LidarPointCloudMeshing.h
--- a/Engine/Plugins/Enterprise/LidarPointCloud/Source/LidarPointCloudRuntime/Public/Meshing/LidarPointCloudMeshing.h
+++ b/Engine/Plugins/Enterprise/LidarPointCloud/Source/LidarPointCloudRuntime/Public/Meshing/LidarPointCloudMeshing.h
@@ -4,6 +4,7 @@
 
 #include "CoreMinimal.h"
 #include "LidarPointCloudShared.h"
+#include "HAL/ThreadSafeBool.h"
 
 class FLidarPointCloudOctree;
 struct FTriMeshCollisionData;
