Fixed build break.

diff --git a/Engine/Plugins/Enterprise/LidarPointCloud/Source/LidarPointCloudRuntime/Private/Rendering/LidarPointCloudRenderBuffers.cpp b/Engine/Plugins/Enterprise/LidarPointCloud/Source/LidarPointCloudRuntime/Private/Rendering/LidarPointCloudRenderBuffers.cpp
--- a/Engine/Plugins/Enterprise/LidarPointCloud/Source/LidarPointCloudRuntime/Private/Rendering/LidarPointCloudRenderBuffers.cpp
+++ b/Engine/Plugins/Enterprise/LidarPointCloud/Source/LidarPointCloudRuntime/Private/Rendering/LidarPointCloudRenderBuffers.cpp
@@ -350,8 +350,7 @@ void FLidarPointCloudVertexFactory::FPointCloudVertexBuffer::InitRHI()
 
 void FLidarPointCloudVertexFactory::InitRHI()
 {
-	FRHICommandListBase& RHICmdList = FRHICommandListImmediate::Get();
-	VertexBuffer.InitResource(RHICmdList);
+	VertexBuffer.InitResource(FRHICommandListImmediate::Get());
 	
 	FVertexDeclarationElementList Elements;
 	Elements.Add(AccessStreamComponent(FVertexStreamComponent(&VertexBuffer, 0, sizeof(FLidarPointCloudPoint), VET_Float3), 0));
@@ -359,7 +358,7 @@ void FLidarPointCloudVertexFactory::InitRHI()
 	Elements.Add(AccessStreamComponent(FVertexStreamComponent(&VertexBuffer, 16, sizeof(FLidarPointCloudPoint), VET_UInt), 2));
 	InitDeclaration(Elements);
 
-	FLidarPointCloudVertexFactoryBase::InitRHI(RHICmdList);
+	FLidarPointCloudVertexFactoryBase::InitRHI();
 }
 
 void FLidarPointCloudVertexFactory::ReleaseRHI()
