[On behalf of Michal Cieciura]
Workaround for serialization issues in Lidar assets

#rnx

diff --git a/Engine/Plugins/Enterprise/LidarPointCloud/Source/LidarPointCloudRuntime/Private/LidarPointCloudOctree.cpp b/Engine/Plugins/Enterprise/LidarPointCloud/Source/LidarPointCloudRuntime/Private/LidarPointCloudOctree.cpp
--- a/Engine/Plugins/Enterprise/LidarPointCloud/Source/LidarPointCloudRuntime/Private/LidarPointCloudOctree.cpp
+++ b/Engine/Plugins/Enterprise/LidarPointCloud/Source/LidarPointCloudRuntime/Private/LidarPointCloudOctree.cpp
@@ -2765,7 +2765,7 @@ void FLidarPointCloudOctree::Serialize(FArchive& Ar)
 		}
 		else if(Ar.IsSaving())
 		{
-			SavingBulkData.Serialize(Ar, Owner, false, 1, EFileRegionType::None);
+			SavingBulkData.Serialize(Ar, Owner);
 		}
 		else
 #endif
@@ -2874,8 +2874,17 @@ FLidarPointCloudOctree::FLidarPointCloudBulkData::FLidarPointCloudBulkData(FLida
 	
 	SetBulkDataFlags(BULKDATA_Force_NOT_InlinePayload | BULKDATA_Size64Bit);
 }
-#endif
 
+void FLidarPointCloudOctree::FLidarPointCloudBulkData::Serialize(FArchive& Ar, UObject* InOwner)
+{
+	Lock(LOCK_READ_WRITE);
+	Realloc(1,1);
+	Unlock();
+	SetBulkDataFlags(BULKDATA_Force_NOT_InlinePayload | BULKDATA_Size64Bit);
+
+	FBulkData::Serialize(Ar, InOwner, false, 1, EFileRegionType::None);
+}
+#endif
 //////////////////////////////////////////////////////////// FLidarPointCloudTraversalOctreeNode
 
 FLidarPointCloudTraversalOctreeNode::FLidarPointCloudTraversalOctreeNode()
diff --git a/Engine/Plugins/Enterprise/LidarPointCloud/Source/LidarPointCloudRuntime/Public/LidarPointCloudOctree.h b/Engine/Plugins/Enterprise/LidarPointCloud/Source/LidarPointCloudRuntime/Public/LidarPointCloudOctree.h
--- a/Engine/Plugins/Enterprise/LidarPointCloud/Source/LidarPointCloudRuntime/Public/LidarPointCloudOctree.h
+++ b/Engine/Plugins/Enterprise/LidarPointCloud/Source/LidarPointCloudRuntime/Public/LidarPointCloudOctree.h
@@ -289,6 +289,8 @@ private:
 		FSerializeBulkDataElements SerializeElementsCallback;
 		
 		FLidarPointCloudBulkData(FLidarPointCloudOctree* Octree);
+
+		void Serialize(FArchive& Ar, UObject* InOwner);
 	} SavingBulkData;
 #endif
 	
