Lidar Point Cloud: Fixed FSceneViews leaking because of incorrect use of FSceneViewFamily / FSceneViewFamilyContext.

diff --git a/Engine/Plugins/Enterprise/LidarPointCloud/Source/LidarPointCloudEditor/Private/LidarPointCloudEditorHelper.cpp b/Engine/Plugins/Enterprise/LidarPointCloud/Source/LidarPointCloudEditor/Private/LidarPointCloudEditorHelper.cpp
--- a/Engine/Plugins/Enterprise/LidarPointCloud/Source/LidarPointCloudEditor/Private/LidarPointCloudEditorHelper.cpp
+++ b/Engine/Plugins/Enterprise/LidarPointCloud/Source/LidarPointCloudEditor/Private/LidarPointCloudEditorHelper.cpp
@@ -376,13 +376,14 @@ namespace
 		}
 	}
 	
-	FSceneView* GetEditorView(FEditorViewportClient* ViewportClient)
+	TSharedPtr<FSceneView> GetEditorView(FEditorViewportClient* ViewportClient)
 	{
-		FSceneViewFamilyContext ViewFamily(FSceneViewFamily::ConstructionValues(
+		FSceneViewFamily ViewFamily(FSceneViewFamily::ConstructionValues(
 			ViewportClient->Viewport,
 			ViewportClient->GetScene(),
 			ViewportClient->EngineShowFlags).SetRealtimeUpdate(ViewportClient->IsRealtime()));
-		return ViewportClient->CalcSceneView(&ViewFamily);
+
+		return MakeShareable<FSceneView>(ViewportClient->CalcSceneView(&ViewFamily));
 	}
 
 	FConvexVolume BuildConvexVolumeForPoints(const TArray<FVector2D>& Points, FEditorViewportClient* ViewportClient)
@@ -394,7 +395,7 @@ namespace
 		
 		FConvexVolume ConvexVolume;
 
-		const FSceneView* View = GetEditorView(ViewportClient);
+		const TSharedPtr<const FSceneView> View = GetEditorView(ViewportClient);
 		const FMatrix InvViewProjectionMatrix = View->ViewMatrices.GetInvViewProjectionMatrix();
 
 		TArray<FVector> Origins; Origins.AddUninitialized(Points.Num() + 2);
@@ -819,7 +820,7 @@ FLidarPointCloudRay FLidarPointCloudEditorHelper::MakeRayFromScreenPosition(FVec
 		ViewportClient = GCurrentLevelEditingViewportClient;
 	}
 	
-	const FSceneView* View = GetEditorView(ViewportClient);
+	const TSharedPtr<const FSceneView> View = GetEditorView(ViewportClient);
 	const FMatrix InvViewProjectionMatrix = View->ViewMatrices.GetInvViewProjectionMatrix();
 
 	FVector3d Origin, Direction;
diff --git a/Engine/Plugins/Enterprise/LidarPointCloud/Source/LidarPointCloudRuntime/Private/LidarPointCloudLODManager.cpp b/Engine/Plugins/Enterprise/LidarPointCloud/Source/LidarPointCloudRuntime/Private/LidarPointCloudLODManager.cpp
--- a/Engine/Plugins/Enterprise/LidarPointCloud/Source/LidarPointCloudRuntime/Private/LidarPointCloudLODManager.cpp
+++ b/Engine/Plugins/Enterprise/LidarPointCloud/Source/LidarPointCloudRuntime/Private/LidarPointCloudLODManager.cpp
@@ -155,7 +155,7 @@ bool FLidarPointCloudViewData::ComputeFromEditorViewportClient(FViewportClient*
 		{
 			FSceneViewFamily::ConstructionValues CVS(nullptr, nullptr, FEngineShowFlags(EShowFlagInitMode::ESFIM_Game));
 			CVS.SetTime(FGameTime());
-			FSceneViewFamily ViewFamily(CVS);
+			FSceneViewFamilyContext ViewFamily(CVS);
 			FSceneView* View = Client->CalcSceneView(&ViewFamily);
 
 			const FMatrix& ProjectionMatrix = View->ViewMatrices.GetProjectionMatrix();
