Fixed crash  on startup in packaged game on Mac when the Lidar Point Cloud plugin is enabled

#rnx

diff --git a/Engine/Plugins/Enterprise/LidarPointCloud/Source/LidarPointCloudRuntime/Private/IO/LidarPointCloudFileIO.cpp b/Engine/Plugins/Enterprise/LidarPointCloud/Source/LidarPointCloudRuntime/Private/IO/LidarPointCloudFileIO.cpp
--- a/Engine/Plugins/Enterprise/LidarPointCloud/Source/LidarPointCloudRuntime/Private/IO/LidarPointCloudFileIO.cpp
+++ b/Engine/Plugins/Enterprise/LidarPointCloud/Source/LidarPointCloudRuntime/Private/IO/LidarPointCloudFileIO.cpp
@@ -250,9 +250,12 @@ TArray<FString> ULidarPointCloudFileIO::GetSupportedExportExtensions()
 
 void ULidarPointCloudFileIO::RegisterHandler(FLidarPointCloudFileIOHandler* Handler, const TArray<FString>& Extensions)
 {
-	for (const FString& Extension : Extensions)
+	if (Instance)
 	{
-		Instance->RegisteredHandlers.Emplace(Extension, Handler);
+		for (const FString& Extension : Extensions)
+		{
+			Instance->RegisteredHandlers.Emplace(Extension, Handler);
+		}
 	}
 }
 
