[HWRT] Added missing implementation of HasRayTracingRepresentation to some proxies types.

diff --git a/Engine/Plugins/Enterprise/LidarPointCloud/Source/LidarPointCloudRuntime/Private/Rendering/LidarPointCloudRendering.cpp b/Engine/Plugins/Enterprise/LidarPointCloud/Source/LidarPointCloudRuntime/Private/Rendering/LidarPointCloudRendering.cpp
--- a/Engine/Plugins/Enterprise/LidarPointCloud/Source/LidarPointCloudRuntime/Private/Rendering/LidarPointCloudRendering.cpp
+++ b/Engine/Plugins/Enterprise/LidarPointCloud/Source/LidarPointCloudRuntime/Private/Rendering/LidarPointCloudRendering.cpp
@@ -415,6 +415,7 @@ public:
 
 #if RHI_RAYTRACING
 	virtual bool IsRayTracingRelevant() const override { return true; }
+	virtual bool HasRayTracingRepresentation() const override { return true; }
 
 	virtual void GetDynamicRayTracingInstances(FRayTracingInstanceCollector& Collector) override
 	{
diff --git a/Engine/Source/Runtime/Engine/Private/Components/HeterogeneousVolumeComponent.cpp b/Engine/Source/Runtime/Engine/Private/Components/HeterogeneousVolumeComponent.cpp
--- a/Engine/Source/Runtime/Engine/Private/Components/HeterogeneousVolumeComponent.cpp
+++ b/Engine/Source/Runtime/Engine/Private/Components/HeterogeneousVolumeComponent.cpp
@@ -45,6 +45,7 @@ public:
 	}
 #if RHI_RAYTRACING
 	virtual bool IsRayTracingRelevant() const override  { return true; }
+	virtual bool HasRayTracingRepresentation() const override { return true; }
 #endif // RHI_RAYTRACING
 	virtual uint32 GetMemoryFootprint(void) const override { return sizeof(*this) + GetAllocatedSize(); }
 	uint32 GetAllocatedSize(void) const { return FPrimitiveSceneProxy::GetAllocatedSize(); }
