Fixed non unity compile errors

diff --git a/Engine/Plugins/Enterprise/LidarPointCloud/Source/LidarPointCloudRuntime/Private/LidarPointCloudLODManager.cpp b/Engine/Plugins/Enterprise/LidarPointCloud/Source/LidarPointCloudRuntime/Private/LidarPointCloudLODManager.cpp
--- a/Engine/Plugins/Enterprise/LidarPointCloud/Source/LidarPointCloudRuntime/Private/LidarPointCloudLODManager.cpp
+++ b/Engine/Plugins/Enterprise/LidarPointCloud/Source/LidarPointCloudRuntime/Private/LidarPointCloudLODManager.cpp
@@ -14,6 +14,7 @@
 #include "Engine/Engine.h"
 #include "Engine/GameViewportClient.h"
 #include "EngineUtils.h"
+#include "SceneView.h"
 #include "UnrealClient.h"
 
 #if WITH_EDITOR
diff --git a/Engine/Plugins/Enterprise/LidarPointCloud/Source/LidarPointCloudRuntime/Private/Rendering/LidarPointCloudRendering.cpp b/Engine/Plugins/Enterprise/LidarPointCloud/Source/LidarPointCloudRuntime/Private/Rendering/LidarPointCloudRendering.cpp
--- a/Engine/Plugins/Enterprise/LidarPointCloud/Source/LidarPointCloudRuntime/Private/Rendering/LidarPointCloudRendering.cpp
+++ b/Engine/Plugins/Enterprise/LidarPointCloud/Source/LidarPointCloudRuntime/Private/Rendering/LidarPointCloudRendering.cpp
@@ -15,6 +15,8 @@
 #include "LocalVertexFactory.h"
 #include "Materials/Material.h"
 #include "Materials/MaterialRenderProxy.h"
+#include "SceneView.h"
+
 #if RHI_RAYTRACING
 #include "RayTracingInstance.h"
 #endif
diff --git a/Engine/Plugins/Enterprise/LidarPointCloud/Source/LidarPointCloudRuntime/Public/LidarPointCloudLODManager.h b/Engine/Plugins/Enterprise/LidarPointCloud/Source/LidarPointCloudRuntime/Public/LidarPointCloudLODManager.h
--- a/Engine/Plugins/Enterprise/LidarPointCloud/Source/LidarPointCloudRuntime/Public/LidarPointCloudLODManager.h
+++ b/Engine/Plugins/Enterprise/LidarPointCloud/Source/LidarPointCloudRuntime/Public/LidarPointCloudLODManager.h
@@ -2,6 +2,7 @@
 
 #pragma once
 
+#include "ConvexVolume.h"
 #include "CoreMinimal.h"
 #include "Tickable.h"
 #include "PrimitiveSceneProxy.h"
diff --git a/Engine/Plugins/Experimental/AnimNext/Source/AnimNext/Private/AnimNextDebugDraw.cpp b/Engine/Plugins/Experimental/AnimNext/Source/AnimNext/Private/AnimNextDebugDraw.cpp
--- a/Engine/Plugins/Experimental/AnimNext/Source/AnimNext/Private/AnimNextDebugDraw.cpp
+++ b/Engine/Plugins/Experimental/AnimNext/Source/AnimNext/Private/AnimNextDebugDraw.cpp
@@ -9,6 +9,7 @@
 #include "Engine/EngineTypes.h"
 #include "Engine/World.h"
 #include "SceneManagement.h"
+#include "SceneView.h"
 #include "Module/AnimNextModuleInstance.h"
 
 namespace UE::AnimNext::Debug
diff --git a/Engine/Plugins/Runtime/CustomMeshComponent/Source/CustomMeshComponent/Private/CustomMeshComponent.cpp b/Engine/Plugins/Runtime/CustomMeshComponent/Source/CustomMeshComponent/Private/CustomMeshComponent.cpp
--- a/Engine/Plugins/Runtime/CustomMeshComponent/Source/CustomMeshComponent/Private/CustomMeshComponent.cpp
+++ b/Engine/Plugins/Runtime/CustomMeshComponent/Source/CustomMeshComponent/Private/CustomMeshComponent.cpp
@@ -12,6 +12,7 @@
 #include "DynamicMeshBuilder.h"
 #include "Engine/Engine.h"
 #include "SceneInterface.h"
+#include "SceneView.h"
 #include "StaticMeshResources.h"
 
 /** Scene proxy */
diff --git a/Engine/Source/Runtime/Renderer/Public/MeshPassProcessor.h b/Engine/Source/Runtime/Renderer/Public/MeshPassProcessor.h
--- a/Engine/Source/Runtime/Renderer/Public/MeshPassProcessor.h
+++ b/Engine/Source/Runtime/Renderer/Public/MeshPassProcessor.h
@@ -6,6 +6,7 @@
 
 #pragma once
 
+#include "DebugViewModeHelpers.h"
 #include "MeshDrawShaderBindings.h"
 #include "MeshMaterialShader.h"
 #include "RendererInterface.h"
