Fix ULidarPointCloudFileIO::RegisterHandler not always Registering Handles due to static initialization order

diff --git a/Engine/Plugins/Enterprise/LidarPointCloud/Source/LidarPointCloudRuntime/Private/IO/LidarPointCloudFileIO.cpp b/Engine/Plugins/Enterprise/LidarPointCloud/Source/LidarPointCloudRuntime/Private/IO/LidarPointCloudFileIO.cpp
--- a/Engine/Plugins/Enterprise/LidarPointCloud/Source/LidarPointCloudRuntime/Private/IO/LidarPointCloudFileIO.cpp
+++ b/Engine/Plugins/Enterprise/LidarPointCloud/Source/LidarPointCloudRuntime/Private/IO/LidarPointCloudFileIO.cpp
@@ -6,7 +6,7 @@
 #include "Misc/Paths.h"
 #include "HAL/FileManager.h"
 
-ULidarPointCloudFileIO* ULidarPointCloudFileIO::Instance = nullptr;
+TMap<FString, FLidarPointCloudFileIOHandler*> ULidarPointCloudFileIO::RegisteredHandlers;
 
 //////////////////////////////////////////////////////////// Settings
 
@@ -222,7 +222,7 @@ TArray<FString> ULidarPointCloudFileIO::GetSupportedImportExtensions()
 {
 	TArray<FString> Extensions;
 
-	for (const TPair<FString, FLidarPointCloudFileIOHandler*>& Handler : Instance->RegisteredHandlers)
+	for (const TPair<FString, FLidarPointCloudFileIOHandler*>& Handler : RegisteredHandlers)
 	{
 		if (Handler.Value->SupportsImport())
 		{
@@ -237,7 +237,7 @@ TArray<FString> ULidarPointCloudFileIO::GetSupportedExportExtensions()
 {
 	TArray<FString> Extensions;
 
-	for (const TPair<FString, FLidarPointCloudFileIOHandler*>& Handler : Instance->RegisteredHandlers)
+	for (const TPair<FString, FLidarPointCloudFileIOHandler*>& Handler : RegisteredHandlers)
 	{
 		if (Handler.Value->SupportsExport())
 		{
@@ -250,18 +250,15 @@ TArray<FString> ULidarPointCloudFileIO::GetSupportedExportExtensions()
 
 void ULidarPointCloudFileIO::RegisterHandler(FLidarPointCloudFileIOHandler* Handler, const TArray<FString>& Extensions)
 {
-	if (Instance)
+	for (const FString& Extension : Extensions)
 	{
-		for (const FString& Extension : Extensions)
-		{
-			Instance->RegisteredHandlers.Emplace(Extension, Handler);
-		}
+		RegisteredHandlers.Emplace(Extension, Handler);
 	}
 }
 
 FLidarPointCloudFileIOHandler* ULidarPointCloudFileIO::FindHandlerByType(const FString& Type)
 {
-	FLidarPointCloudFileIOHandler** Handler = Instance->RegisteredHandlers.Find(Type);
+	FLidarPointCloudFileIOHandler** Handler = RegisteredHandlers.Find(Type);
 	return Handler ? *Handler : nullptr;
 }
 
@@ -314,14 +311,8 @@ ULidarPointCloudFileIO::ULidarPointCloudFileIO(const FObjectInitializer& ObjectI
 	SupportedClass = ULidarPointCloud::StaticClass();
 	PreferredFormatIndex = 0;
 
-	// This will assign an Instance pointer
-	if (!Instance)
-	{
-		Instance = this;
-	}
-
 	// Requested by UExporter
-	else
+	if(!HasAnyFlags(RF_ClassDefaultObject))
 	{
 		FormatExtension.Append(ULidarPointCloudFileIO::GetSupportedExportExtensions());
 		for (int32 i = 0; i < FormatExtension.Num(); i++)
diff --git a/Engine/Plugins/Enterprise/LidarPointCloud/Source/LidarPointCloudRuntime/Public/IO/LidarPointCloudFileIO.h b/Engine/Plugins/Enterprise/LidarPointCloud/Source/LidarPointCloudRuntime/Public/IO/LidarPointCloudFileIO.h
--- a/Engine/Plugins/Enterprise/LidarPointCloud/Source/LidarPointCloudRuntime/Public/IO/LidarPointCloudFileIO.h
+++ b/Engine/Plugins/Enterprise/LidarPointCloud/Source/LidarPointCloudRuntime/Public/IO/LidarPointCloudFileIO.h
@@ -187,10 +187,8 @@ class LIDARPOINTCLOUDRUNTIME_API ULidarPointCloudFileIO : public UExporter
 	GENERATED_UCLASS_BODY()
 
 private:
-	static ULidarPointCloudFileIO* Instance;
-
 	/** Links extensions to file handlers. */
-	TMap<FString, FLidarPointCloudFileIOHandler*> RegisteredHandlers;
+	static TMap<FString, FLidarPointCloudFileIOHandler*> RegisteredHandlers;
 
 public:
 	/**
