Replace some usages of FORCEINLINE with inline in Enterprise.

For context about this change, see this doc: https://docs.google.com/document/d/1LjpT7NpzTecfXdFNyv5KM6sCu67827BRljNi1iaWVSE/edit?tab=t.0#heading=h.fvtz8gjukz7x

Reviewers were automatically determined. Please comment if anyone else should be added.

#rnx

diff --git a/Engine/Plugins/Enterprise/LidarPointCloud/Source/LidarPointCloudRuntime/Public/IO/LidarPointCloudFileIO.h b/Engine/Plugins/Enterprise/LidarPointCloud/Source/LidarPointCloudRuntime/Public/IO/LidarPointCloudFileIO.h
--- a/Engine/Plugins/Enterprise/LidarPointCloud/Source/LidarPointCloudRuntime/Public/IO/LidarPointCloudFileIO.h
+++ b/Engine/Plugins/Enterprise/LidarPointCloud/Source/LidarPointCloudRuntime/Public/IO/LidarPointCloudFileIO.h
@@ -109,7 +109,7 @@ public:
 
 	void SetPointCount(const uint64& InTotalPointCount);
 
-	FORCEINLINE void AddPoint(const FVector& Location, const float& R, const float& G, const float& B, const float& A = 1.0f)
+	inline void AddPoint(const FVector& Location, const float& R, const float& G, const float& B, const float& A = 1.0f)
 	{
 		AddPoint(Location.X, Location.Y, Location.Z, R, G, B, A);
 	}
@@ -122,7 +122,7 @@ public:
 
 	void CenterPoints();
 
-	FORCEINLINE bool IsCancelled() const
+	inline bool IsCancelled() const
 	{
 		return bCancelled && *bCancelled;
 	}
diff --git a/Engine/Plugins/Enterprise/LidarPointCloud/Source/LidarPointCloudRuntime/Public/IO/LidarPointCloudFileIO_ASCII.h b/Engine/Plugins/Enterprise/LidarPointCloud/Source/LidarPointCloudRuntime/Public/IO/LidarPointCloudFileIO_ASCII.h
--- a/Engine/Plugins/Enterprise/LidarPointCloud/Source/LidarPointCloudRuntime/Public/IO/LidarPointCloudFileIO_ASCII.h
+++ b/Engine/Plugins/Enterprise/LidarPointCloud/Source/LidarPointCloudRuntime/Public/IO/LidarPointCloudFileIO_ASCII.h
@@ -129,7 +129,7 @@ public:
 		return NewSettings;
 	}
 
-	FORCEINLINE bool HasValidData() { return !Delimiter.IsEmpty() && Columns.Num() > 0; }
+	inline bool HasValidData() { return !Delimiter.IsEmpty() && Columns.Num() > 0; }
 
 private:
 	/** Reads and parses header information about the given file. */
diff --git a/Engine/Plugins/Enterprise/LidarPointCloud/Source/LidarPointCloudRuntime/Public/LidarPointCloud.h b/Engine/Plugins/Enterprise/LidarPointCloud/Source/LidarPointCloudRuntime/Public/LidarPointCloud.h
--- a/Engine/Plugins/Enterprise/LidarPointCloud/Source/LidarPointCloudRuntime/Public/LidarPointCloud.h
+++ b/Engine/Plugins/Enterprise/LidarPointCloud/Source/LidarPointCloudRuntime/Public/LidarPointCloud.h
@@ -203,7 +203,7 @@ public:
 	int32 GetNumNodes() const { return Octree.GetNumNodes(); }
 
 	UFUNCTION(BlueprintPure, Category = "Lidar Point Cloud")
-	FORCEINLINE float GetEstimatedPointSpacing() const { return Octree.GetEstimatedPointSpacing(); }
+	inline float GetEstimatedPointSpacing() const { return Octree.GetEstimatedPointSpacing(); }
 
 	/** Returns the amount of memory in MB used to store the point cloud. */
 	UFUNCTION(BlueprintPure, Category = "Lidar Point Cloud")
diff --git a/Engine/Plugins/Enterprise/LidarPointCloud/Source/LidarPointCloudRuntime/Public/LidarPointCloudComponent.h b/Engine/Plugins/Enterprise/LidarPointCloud/Source/LidarPointCloudRuntime/Public/LidarPointCloudComponent.h
--- a/Engine/Plugins/Enterprise/LidarPointCloud/Source/LidarPointCloudRuntime/Public/LidarPointCloudComponent.h
+++ b/Engine/Plugins/Enterprise/LidarPointCloud/Source/LidarPointCloudRuntime/Public/LidarPointCloudComponent.h
@@ -160,8 +160,8 @@ public:
 	UFUNCTION(BlueprintPure, Category = "Components|LidarPointCloud")
 	ULidarPointCloud* GetPointCloud() const { return PointCloud; }
 
-	FORCEINLINE TWeakPtr<FViewportClient> GetOwningViewportClient() const { return OwningViewportClient; }
-	FORCEINLINE bool IsOwnedByEditor() const { return OwningViewportClient.IsValid(); }
+	inline TWeakPtr<FViewportClient> GetOwningViewportClient() const { return OwningViewportClient; }
+	inline bool IsOwnedByEditor() const { return OwningViewportClient.IsValid(); }
 
 	/** Returns true if there are any points within the given sphere. */
 	UFUNCTION(BlueprintPure, Category = "Lidar Point Cloud")
@@ -519,7 +519,7 @@ public:
 
 	/** Returns the current Point Shape */
 	UFUNCTION(BlueprintPure, Category = "Components|LidarPointCloud")
-	FORCEINLINE ELidarPointCloudSpriteShape GetPointShape() const { return PointShape; }
+	inline ELidarPointCloudSpriteShape GetPointShape() const { return PointShape; }
 
 	/** Sets new Point Shape */
 	UFUNCTION(BlueprintCallable, Category = "Components|LidarPointCloud")
diff --git a/Engine/Plugins/Enterprise/LidarPointCloud/Source/LidarPointCloudRuntime/Public/LidarPointCloudOctree.h b/Engine/Plugins/Enterprise/LidarPointCloud/Source/LidarPointCloudRuntime/Public/LidarPointCloudOctree.h
--- a/Engine/Plugins/Enterprise/LidarPointCloud/Source/LidarPointCloudRuntime/Public/LidarPointCloudOctree.h
+++ b/Engine/Plugins/Enterprise/LidarPointCloud/Source/LidarPointCloudRuntime/Public/LidarPointCloudOctree.h
@@ -116,7 +116,7 @@ private:
 	TAtomic<bool> bCanReleaseData;
 
 public:
-	FORCEINLINE FLidarPointCloudOctreeNode(FLidarPointCloudOctree* Tree, const uint8& Depth) : FLidarPointCloudOctreeNode(Tree, Depth, 0, FVector3f::ZeroVector) {}
+	inline FLidarPointCloudOctreeNode(FLidarPointCloudOctree* Tree, const uint8& Depth) : FLidarPointCloudOctreeNode(Tree, Depth, 0, FVector3f::ZeroVector) {}
 	FLidarPointCloudOctreeNode(FLidarPointCloudOctree* Tree, const uint8& Depth, const uint8& LocationInParent, const FVector3f& Center);
 	~FLidarPointCloudOctreeNode();
 	FLidarPointCloudOctreeNode(const FLidarPointCloudOctreeNode&) = delete;
@@ -131,13 +131,13 @@ public:
 	FLidarPointCloudPoint* GetPersistentData() const;
 
 	/** Returns a pointer to the point data */
-	FORCEINLINE TSharedPtr<FLidarPointCloudRenderBuffer> GetDataCache() { return DataCache; }
+	inline TSharedPtr<FLidarPointCloudRenderBuffer> GetDataCache() { return DataCache; }
 
 	/** Return a pointer to the vertex factory containing pre-cached geometry */
-	FORCEINLINE TSharedPtr<FLidarPointCloudVertexFactory> GetVertexFactory() { return VertexFactory; }
+	inline TSharedPtr<FLidarPointCloudVertexFactory> GetVertexFactory() { return VertexFactory; }
 	
 	/** Return a pointer to the ray tracing geometry */
-	FORCEINLINE TSharedPtr<FLidarPointCloudRayTracingGeometry> GetRayTracingGeometry() { return RayTracingGeometry; }
+	inline TSharedPtr<FLidarPointCloudRayTracingGeometry> GetRayTracingGeometry() { return RayTracingGeometry; }
 
 	/**
 	 * Builds and updates the necessary render data buffers
@@ -146,16 +146,16 @@ public:
 	bool BuildDataCache(bool bUseStaticBuffers, bool bUseRayTracing);
 
 	/** Returns the sum of grid and padding points allocated to this node. */
-	FORCEINLINE uint32 GetNumPoints() const { return NumPoints; }
+	inline uint32 GetNumPoints() const { return NumPoints; }
 
 	/** Returns the sum of visible grid and padding points allocated to this node. */
 	uint32 GetNumVisiblePoints() const { return NumVisiblePoints; }
 
 	/** Calculates and returns the bounds of this node */
-	FORCEINLINE FBox GetBounds() const;
+	inline FBox GetBounds() const;
 
 	/** Calculates and returns the sphere bounds of this node */
-	FORCEINLINE FSphere GetSphereBounds() const;
+	inline FSphere GetSphereBounds() const;
 
 	/** Returns a pointer to the node at the given location, or null if one doesn't exist yet. */
 	FLidarPointCloudOctreeNode* GetChildNodeAtLocation(const uint8& Location) const;
@@ -321,7 +321,7 @@ public:
 	FBox GetBounds() const { return FBox(-Extent, Extent); }
 
 	/** Returns the extent of the Cloud's bounds. */
-	FORCEINLINE FVector3f GetExtent() const { return Extent; }
+	inline FVector3f GetExtent() const { return Extent; }
 
 	/** Recalculates and updates points bounds. */
 	void RefreshBounds();
@@ -336,7 +336,7 @@ public:
 	int32 GetNumNodes() const;
 
 	/** Returns the total number of nodes. */
-	FORCEINLINE int32 GetNumNodesInUse() const { return NodesInUse.Num(); }
+	inline int32 GetNumNodesInUse() const { return NodesInUse.Num(); }
 
 	/** Returns a pointer to the Point Cloud asset, which owns this Octree. */
 	ULidarPointCloud* GetOwner() const { return Owner; }
@@ -677,7 +677,7 @@ struct FLidarPointCloudTraversalOctreeNode
 	/** Calculates virtual depth of this node, to be used to estimate the best sprite size */
 	void CalculateVirtualDepth(const TArray<float>& LevelWeights, const float& PointSizeBias);
 
-	FORCEINLINE bool IsAvailable() const { return bSelected && DataNode->HasRenderData(); }
+	inline bool IsAvailable() const { return bSelected && DataNode->HasRenderData(); }
 };
 
 /** Used for node size sorting and node selection. */
diff --git a/Engine/Plugins/Enterprise/LidarPointCloud/Source/LidarPointCloudRuntime/Public/LidarPointCloudShared.h b/Engine/Plugins/Enterprise/LidarPointCloud/Source/LidarPointCloudRuntime/Public/LidarPointCloudShared.h
--- a/Engine/Plugins/Enterprise/LidarPointCloud/Source/LidarPointCloudRuntime/Public/LidarPointCloudShared.h
+++ b/Engine/Plugins/Enterprise/LidarPointCloud/Source/LidarPointCloudRuntime/Public/LidarPointCloudShared.h
@@ -38,25 +38,25 @@ public:
 
 	bool operator==(const FLidarPointCloudNormal& Other) const { return X == Other.X && Y == Other.Y && Z == Other.Z; }
 
-	FORCEINLINE bool IsValid() const { return X != 127 || Y != 127 || Z != 127; }
+	inline bool IsValid() const { return X != 127 || Y != 127 || Z != 127; }
 
-	FORCEINLINE void SetFromVector(const FVector3f& Normal)
+	inline void SetFromVector(const FVector3f& Normal)
 	{
 		SetFromFloats(Normal.X, Normal.Y, Normal.Z);
 	}
-	FORCEINLINE void SetFromFloats(const float& InX, const float& InY, const float& InZ)
+	inline void SetFromFloats(const float& InX, const float& InY, const float& InZ)
 	{
 		X = FMath::Min((InX + 1) * 127.5f, 255.0f);
 		Y = FMath::Min((InY + 1) * 127.5f, 255.0f);
 		Z = FMath::Min((InZ + 1) * 127.5f, 255.0f);
 	}
 
-	FORCEINLINE void Reset()
+	inline void Reset()
 	{
 		X = Y = Z = 127;
 	}
 
-	FORCEINLINE FVector3f ToVector() const { return FVector3f(X / 127.5f - 1, Y / 127.5f - 1, Z / 127.5f - 1); }
+	inline FVector3f ToVector() const { return FVector3f(X / 127.5f - 1, Y / 127.5f - 1, Z / 127.5f - 1); }
 };
 
 /** Used for backwards compatibility with pre-normal datasets */
@@ -173,7 +173,7 @@ public:
 	{
 	}
 
-	FORCEINLINE void CopyFrom(const FLidarPointCloudPoint& Other)
+	inline void CopyFrom(const FLidarPointCloudPoint& Other)
 	{
 		Location = Other.Location;
 		Color = Other.Color;
@@ -182,7 +182,7 @@ public:
 		ClassificationID = Other.ClassificationID;
 	}
 
-	FORCEINLINE FLidarPointCloudPoint Transform(const FTransform3f& Transform) const
+	inline FLidarPointCloudPoint Transform(const FTransform3f& Transform) const
 	{
 		return FLidarPointCloudPoint(Transform.TransformPosition(Location), Color, bVisible, ClassificationID);
 	}
@@ -255,8 +255,8 @@ public:
 	}
 	FLidarPointCloudDataBuffer& operator=(FLidarPointCloudDataBuffer&&) = delete;
 
-	FORCEINLINE uint8* GetData() { return Data.GetData(); }
-	FORCEINLINE bool InUse() const { return bInUse; }
+	inline uint8* GetData() { return Data.GetData(); }
+	inline bool InUse() const { return bInUse; }
 	
 	/** Marks the buffer as no longer in use so it can be reassigned to another read thread. */
 	void MarkAsFree();
@@ -313,7 +313,7 @@ public:
 	}
 	FLidarPointCloudRay(const FVector& Origin, const FVector& Direction) : FLidarPointCloudRay((FVector3f)Origin, (FVector3f)Direction) {}
 
-	static FORCEINLINE FLidarPointCloudRay FromLocations(const FVector3f& Origin, const FVector3f& Destination)
+	static inline FLidarPointCloudRay FromLocations(const FVector3f& Origin, const FVector3f& Destination)
 	{
 		return FLidarPointCloudRay(Origin, (Destination - Origin).GetSafeNormal());
 	}
@@ -328,24 +328,24 @@ public:
 	{
 		return FLidarPointCloudRay(Transform.TransformPosition((FVector)Origin), Transform.TransformVector((FVector)Direction));
 	}
-	FORCEINLINE FLidarPointCloudRay ShiftBy(const FVector3f& Offset) const
+	inline FLidarPointCloudRay ShiftBy(const FVector3f& Offset) const
 	{
 		return FLidarPointCloudRay(Origin + Offset, Direction);
 	}
-	FORCEINLINE FLidarPointCloudRay ShiftBy(const FVector& Offset) const
+	inline FLidarPointCloudRay ShiftBy(const FVector& Offset) const
 	{
 		return FLidarPointCloudRay(Origin + (FVector3f)Offset, Direction);
 	}
 
-	FORCEINLINE FVector3f GetDirection() const { return Direction; }
-	FORCEINLINE void SetDirection(const FVector3f& NewDirection)
+	inline FVector3f GetDirection() const { return Direction; }
+	inline void SetDirection(const FVector3f& NewDirection)
 	{
 		Direction = NewDirection;
 		InvDirection = Direction.Reciprocal();
 	}
 
 	/** An Efficient and Robust Ray-Box Intersection Algorithm. Amy Williams et al. 2004. */
-	FORCEINLINE bool Intersects(const FBox& Box) const
+	inline bool Intersects(const FBox& Box) const
 	{
 		float tmin, tmax, tymin, tymax, tzmin, tzmax;
 
@@ -379,7 +379,7 @@ public:
 
 		return true;
 	}
-	FORCEINLINE bool Intersects(const FLidarPointCloudPoint* Point, const float& RadiusSq) const
+	inline bool Intersects(const FLidarPointCloudPoint* Point, const float& RadiusSq) const
 	{
 		const FVector3f L = Point->Location - Origin;
 		const float tca = FVector3f::DotProduct(L, Direction);
@@ -412,7 +412,7 @@ struct FLidarPointCloudClippingVolumeParams
 	FBox Bounds;
 	FMatrix PackedShaderData;
 
-	FORCEINLINE bool operator<(const FLidarPointCloudClippingVolumeParams& O) const
+	inline bool operator<(const FLidarPointCloudClippingVolumeParams& O) const
 	{
 		return (Priority < O.Priority) || (Priority == O.Priority && Mode > O.Mode);
 	}
