Updated PointCloud plugin with more API exposed to BP

#rnx

diff --git a/Engine/Plugins/Enterprise/LidarPointCloud/Source/LidarPointCloudEditor/Private/LidarPointCloudEditorHelper.cpp b/Engine/Plugins/Enterprise/LidarPointCloud/Source/LidarPointCloudEditor/Private/LidarPointCloudEditorHelper.cpp
--- a/Engine/Plugins/Enterprise/LidarPointCloud/Source/LidarPointCloudEditor/Private/LidarPointCloudEditorHelper.cpp
+++ b/Engine/Plugins/Enterprise/LidarPointCloud/Source/LidarPointCloudEditor/Private/LidarPointCloudEditorHelper.cpp
@@ -469,7 +469,7 @@ namespace
 			SelectedPoints.Reserve(NumPoints);
 			ProcessAll([&SelectedPoints](ALidarPointCloudActor* Actor)
 			{
-				Actor->GetPointCloudComponent()->GetSelectedPointsAsCopies(SelectedPoints);
+				Actor->GetPointCloudComponent()->GetSelectedPointsAsCopies(SelectedPoints, true);
 			});
 			
 			NewPointCloud->SetData(SelectedPoints);
diff --git a/Engine/Plugins/Enterprise/LidarPointCloud/Source/LidarPointCloudRuntime/Private/LidarPointCloud.cpp b/Engine/Plugins/Enterprise/LidarPointCloud/Source/LidarPointCloudRuntime/Private/LidarPointCloud.cpp
--- a/Engine/Plugins/Enterprise/LidarPointCloud/Source/LidarPointCloudRuntime/Private/LidarPointCloud.cpp
+++ b/Engine/Plugins/Enterprise/LidarPointCloud/Source/LidarPointCloudRuntime/Private/LidarPointCloud.cpp
@@ -636,12 +636,23 @@ bool ULidarPointCloud::HasSelectedPoints() const
 	return Octree.HasSelectedPoints();
 }
 
-void ULidarPointCloud::GetSelectedPointsAsCopies(TArray64<FLidarPointCloudPoint>& SelectedPoints, FTransform Transform) const
+template <typename T>
+void ULidarPointCloud::GetSelectedPointsAsCopies_Internal(TArray<FLidarPointCloudPoint, T>& SelectedPoints,	FTransform Transform) const
 {
 	Transform.AddToTranslation(LocationOffset);
 	Octree.GetSelectedPointsAsCopies(SelectedPoints, Transform);
 }
 
+void ULidarPointCloud::GetSelectedPointsAsCopies(TArray<FLidarPointCloudPoint>& SelectedPoints, const FTransform& Transform) const
+{
+	GetSelectedPointsAsCopies_Internal(SelectedPoints, Transform);
+}
+
+void ULidarPointCloud::GetSelectedPointsAsCopies(TArray64<FLidarPointCloudPoint>& SelectedPoints, const FTransform& Transform) const
+{
+	GetSelectedPointsAsCopies_Internal(SelectedPoints, Transform);
+}
+
 void ULidarPointCloud::CalculateNormalsForSelection()
 {
 	TArray64<FLidarPointCloudPoint*>* SelectedPoints = new TArray64<FLidarPointCloudPoint*>();
diff --git a/Engine/Plugins/Enterprise/LidarPointCloud/Source/LidarPointCloudRuntime/Private/LidarPointCloudComponent.cpp b/Engine/Plugins/Enterprise/LidarPointCloud/Source/LidarPointCloudRuntime/Private/LidarPointCloudComponent.cpp
--- a/Engine/Plugins/Enterprise/LidarPointCloud/Source/LidarPointCloudRuntime/Private/LidarPointCloudComponent.cpp
+++ b/Engine/Plugins/Enterprise/LidarPointCloud/Source/LidarPointCloudRuntime/Private/LidarPointCloudComponent.cpp
@@ -204,14 +204,6 @@ int64 ULidarPointCloudComponent::NumSelectedPoints()
 	return PointCloud ? PointCloud->NumSelectedPoints() : 0;
 }
 
-void ULidarPointCloudComponent::GetSelectedPointsAsCopies(TArray64<FLidarPointCloudPoint>& SelectedPoints)
-{
-	if(PointCloud)
-	{
-		PointCloud->GetSelectedPointsAsCopies(SelectedPoints, GetComponentTransform());
-	}
-}
-
 void ULidarPointCloudComponent::ClearSelection()
 {
 	if(PointCloud)
diff --git a/Engine/Plugins/Enterprise/LidarPointCloud/Source/LidarPointCloudRuntime/Private/LidarPointCloudOctree.cpp b/Engine/Plugins/Enterprise/LidarPointCloud/Source/LidarPointCloudRuntime/Private/LidarPointCloudOctree.cpp
--- a/Engine/Plugins/Enterprise/LidarPointCloud/Source/LidarPointCloudRuntime/Private/LidarPointCloudOctree.cpp
+++ b/Engine/Plugins/Enterprise/LidarPointCloud/Source/LidarPointCloudRuntime/Private/LidarPointCloudOctree.cpp
@@ -1934,7 +1934,8 @@ void FLidarPointCloudOctree::GetSelectedPoints(TArray64<FLidarPointCloudPoint*>&
 	ITERATE_SELECTED({ SelectedPoints.Add(Point); }, {});
 }
 
-void FLidarPointCloudOctree::GetSelectedPointsAsCopies(TArray64<FLidarPointCloudPoint>& SelectedPoints, const FTransform& Transform) const
+template <typename T>
+void FLidarPointCloudOctree::GetSelectedPointsAsCopies_Internal(TArray<FLidarPointCloudPoint, T>& SelectedPoints, const FTransform& Transform) const
 {
 	const FTransform3f Transform3F = (FTransform3f)Transform;
 
@@ -1944,6 +1945,16 @@ void FLidarPointCloudOctree::GetSelectedPointsAsCopies(TArray64<FLidarPointCloud
 	},{});
 }
 
+void FLidarPointCloudOctree::GetSelectedPointsAsCopies(TArray<FLidarPointCloudPoint>& SelectedPoints, const FTransform& Transform) const
+{
+	GetSelectedPointsAsCopies_Internal(SelectedPoints, Transform);
+}
+
+void FLidarPointCloudOctree::GetSelectedPointsAsCopies(TArray64<FLidarPointCloudPoint>& SelectedPoints, const FTransform& Transform) const
+{
+	GetSelectedPointsAsCopies_Internal(SelectedPoints, Transform);
+}
+
 void FLidarPointCloudOctree::GetSelectedPointsInBox(TArray64<const FLidarPointCloudPoint*>& SelectedPoints, const FBox& Box) const
 {
 	constexpr bool bVisibleOnly = false;
diff --git a/Engine/Plugins/Enterprise/LidarPointCloud/Source/LidarPointCloudRuntime/Public/LidarPointCloud.h b/Engine/Plugins/Enterprise/LidarPointCloud/Source/LidarPointCloudRuntime/Public/LidarPointCloud.h
--- a/Engine/Plugins/Enterprise/LidarPointCloud/Source/LidarPointCloudRuntime/Public/LidarPointCloud.h
+++ b/Engine/Plugins/Enterprise/LidarPointCloud/Source/LidarPointCloudRuntime/Public/LidarPointCloud.h
@@ -431,7 +431,10 @@ public:
 	void InvertSelection();
 	int64 NumSelectedPoints() const;
 	bool HasSelectedPoints() const;
-	void GetSelectedPointsAsCopies(TArray64<FLidarPointCloudPoint>& SelectedPoints, FTransform Transform) const;
+	template <typename T>
+	void GetSelectedPointsAsCopies_Internal(TArray<FLidarPointCloudPoint, T>& SelectedPoints, FTransform Transform) const;
+	void GetSelectedPointsAsCopies(TArray<FLidarPointCloudPoint>& SelectedPoints, const FTransform& Transform) const;
+	void GetSelectedPointsAsCopies(TArray64<FLidarPointCloudPoint>& SelectedPoints, const FTransform& Transform) const;
 	void CalculateNormalsForSelection();
 	void ClearSelection();
 	void BuildStaticMeshBuffersForSelection(float CellSize, LidarPointCloudMeshing::FMeshBuffers* OutMeshBuffers, const FTransform& Transform);
diff --git a/Engine/Plugins/Enterprise/LidarPointCloud/Source/LidarPointCloudRuntime/Public/LidarPointCloudComponent.h b/Engine/Plugins/Enterprise/LidarPointCloud/Source/LidarPointCloudRuntime/Public/LidarPointCloudComponent.h
--- a/Engine/Plugins/Enterprise/LidarPointCloud/Source/LidarPointCloudRuntime/Public/LidarPointCloudComponent.h
+++ b/Engine/Plugins/Enterprise/LidarPointCloud/Source/LidarPointCloudRuntime/Public/LidarPointCloudComponent.h
@@ -196,6 +196,32 @@ public:
 		return PointCloud && PointCloud->HasPointsByRay(Ray.TransformBy(GetComponentTransform().Inverse()), Radius, bVisibleOnly);
 	}
 
+	UFUNCTION(BlueprintPure, Category = "Lidar Point Cloud")
+	TArray<FLidarPointCloudPoint> GetSelectedPointsAsCopies(bool bReturnWorldSpace) const
+	{
+		TArray<FLidarPointCloudPoint> Points;
+#if WITH_EDITOR
+		GetSelectedPointsAsCopies(Points, bReturnWorldSpace);
+#endif
+		return Points;
+	}
+#if WITH_EDITOR
+	void GetSelectedPointsAsCopies(TArray<FLidarPointCloudPoint>& Points, bool bReturnWorldSpace) const
+	{
+		if(PointCloud)
+		{
+			PointCloud->GetSelectedPointsAsCopies(Points, bReturnWorldSpace ? GetComponentTransform() : FTransform());
+		}
+	}
+	void GetSelectedPointsAsCopies(TArray64<FLidarPointCloudPoint>& Points, bool bReturnWorldSpace) const
+	{
+		if(PointCloud)
+		{
+			PointCloud->GetSelectedPointsAsCopies(Points, bReturnWorldSpace ? GetComponentTransform() : FTransform());
+		}
+	}
+#endif
+	
 	/** Populates the array with the list of points within the given sphere. */
 	void GetPointsInSphere(TArray<FLidarPointCloudPoint*>& SelectedPoints, const FVector& Center, const float& Radius, const bool& bVisibleOnly) { GetPointsInSphere(SelectedPoints, FSphere(Center, Radius), bVisibleOnly); }
 	void GetPointsInSphere(TArray64<FLidarPointCloudPoint*>& SelectedPoints, const FVector& Center, const float& Radius, const bool& bVisibleOnly) { GetPointsInSphere(SelectedPoints, FSphere(Center, Radius), bVisibleOnly); }
@@ -243,6 +269,7 @@ public:
 		GetPointsInSphereAsCopies(Points, FSphere(Center, Radius), bVisibleOnly, bReturnWorldSpace);
 		return Points;
 	}
+
 	void GetPointsInSphereAsCopies(TArray<FLidarPointCloudPoint>& SelectedPoints, const FSphere& Sphere, const bool& bVisibleOnly, const bool& bReturnWorldSpace)
 	{
 		if (PointCloud)
@@ -509,7 +536,6 @@ public:
 	void DeleteSelected();
 	void InvertSelection();
 	int64 NumSelectedPoints();
-	void GetSelectedPointsAsCopies(TArray64<FLidarPointCloudPoint>& SelectedPoints);
 	void ClearSelection();
 #endif
 	
diff --git a/Engine/Plugins/Enterprise/LidarPointCloud/Source/LidarPointCloudRuntime/Public/LidarPointCloudOctree.h b/Engine/Plugins/Enterprise/LidarPointCloud/Source/LidarPointCloudRuntime/Public/LidarPointCloudOctree.h
--- a/Engine/Plugins/Enterprise/LidarPointCloud/Source/LidarPointCloudRuntime/Public/LidarPointCloudOctree.h
+++ b/Engine/Plugins/Enterprise/LidarPointCloud/Source/LidarPointCloudRuntime/Public/LidarPointCloudOctree.h
@@ -493,7 +493,10 @@ public:
 	int64 NumSelectedPoints() const;
 	bool HasSelectedPoints() const;
 	void GetSelectedPoints(TArray64<FLidarPointCloudPoint*>& SelectedPoints) const;
-	void GetSelectedPointsAsCopies(TArray64<FLidarPointCloudPoint>& SelectedPoints, const FTransform& Transform) const;
+	template <typename T>
+	void GetSelectedPointsAsCopies_Internal(TArray<FLidarPointCloudPoint, T>& SelectedPoints, const FTransform& Transform) const;
+	void GetSelectedPointsAsCopies(TArray<FLidarPointCloudPoint>& SelectedPoints, const FTransform& Transform) const;	
+	void GetSelectedPointsAsCopies(TArray64<FLidarPointCloudPoint>& SelectedPoints, const FTransform& Transform) const;	
 	void GetSelectedPointsInBox(TArray64<const FLidarPointCloudPoint*>& SelectedPoints, const FBox& Box) const;
 	void ClearSelection();
 	void BuildStaticMeshBuffersForSelection(float CellSize, LidarPointCloudMeshing::FMeshBuffers* OutMeshBuffers, const FTransform& Transform);
